Copying //Tasks/UE4/Dev-Online-Plugins-v3 to //UE4/Dev-Main (Source: //Tasks/UE4/Dev-Online-Plugins-v3 @ 3047294)

==========================
MAJOR FEATURES + CHANGES
==========================

Change 3025090 on 2016/06/23 by [REDACTED_PERSON]

	Move OnlineSubsystemGooglePlay into plugins

Change 3025127 on 2016/06/23 by [REDACTED_PERSON]

	Move OnlineSubsystemIOS into plugins

Change 3025314 on 2016/06/23 by [REDACTED_PERSON]

	Move OnlineSubsystemWeChat into plugins

Change 3025380 on 2016/06/23 by [REDACTED_PERSON]

	Move OnlineSubsystemMcp into plugins

Change 3025400 on 2016/06/23 by [REDACTED_PERSON]

	Move PluginFlow into plugins

Change 3025430 on 2016/06/23 by [REDACTED_PERSON]

	Move OnlineFramework (Hotfix, Lobby, Party, Qos) into plugins

Change 3025453 on 2016/06/23 by [REDACTED_PERSON]

	Move OnlineSubsystem into plugins

Change 3025462 on 2016/06/23 by [REDACTED_PERSON]

	Move OnlineSubsystemAmazon into plugins

Change 3025480 on 2016/06/23 by [REDACTED_PERSON]

	Move OnlineSubsystemFacebook into plugins

Change 3025487 on 2016/06/23 by [REDACTED_PERSON]

	Move OnlineSubsystemNull into plugins

Change 3025493 on 2016/06/23 by [REDACTED_PERSON]

	Move OnlineSubsystemSteam into plugins

Change 3025512 on 2016/06/23 by [REDACTED_PERSON]

	Move OnlineSubsystemUtils (OnlineBlueprintSupport, OnlineSubsystemUtils) into plugins

Change 3025520 on 2016/06/23 by [REDACTED_PERSON]

	Move OnlineSubsystemPS4 into plugins

Change 3025526 on 2016/06/23 by [REDACTED_PERSON]

	Move OnlineSubsystemLive into plugins

Change 3025542 on 2016/06/23 by [REDACTED_PERSON]

	Move OnlineSubsystemThunderhead into plugins (disabled)

Change 3025547 on 2016/06/23 by [REDACTED_PERSON]

	Delete NewsFeed

Change 3028420 on 2016/06/27 by [REDACTED_PERSON]

	Online Subsystem pluginization changes

Change 3028457 on 2016/06/27 by [REDACTED_PERSON]

	Add ShooterOnlineSessionClient.*
	Remove copy/paste from previous checkin

Change 3029282 on 2016/06/27 by [REDACTED_PERSON]

	Merging //UE4/Dev-Main to Dev-Online-Plugins-v3 (//Tasks/UE4/Dev-Online-Plugins-v3) (around CL 3028391)

Change 3029363 on 2016/06/27 by [REDACTED_PERSON]

	UT build fixes for pluginization

Change 3032885 on 2016/06/29 by [REDACTED_PERSON]

	Remove bCompileMCPOss and bCompileSteamOss

Change 3033072 on 2016/06/29 by [REDACTED_PERSON]

	Merging //UE4/Dev-Main to Dev-Online-Plugins-v3 (//Tasks/UE4/Dev-Online-Plugins-v3) (around CL 3032885)

Change 3033368 on 2016/06/29 by [REDACTED_PERSON]

	Fix references to online subsystem files in their former location

Change 3037635 on 2016/07/05 by [REDACTED_PERSON]

	Add online plugins to WEX uproject

Change 3038147 on 2016/07/05 by [REDACTED_PERSON]

	Merging //UE4/Dev-Main to Dev-Online-Plugins-v3 (//Tasks/UE4/Dev-Online-Plugins-v3)
	Also adds Match3

Change 3038152 on 2016/07/05 by [REDACTED_PERSON]

	Fixup UnrealMatch3 OSS references

Change 3044558 on 2016/07/11 by [REDACTED_PERSON]

	Move OnlineBlueprintCallProxyBase to Engine/Public/Net to remove dependency for OnlineSubsystemUtils in Engine.

Change 3044596 on 2016/07/11 by [REDACTED_PERSON]

	Allow users to add a custom FTextLayout to Text boxes

Change 3045306 on 2016/07/11 by [REDACTED_PERSON]

	Merging //UE4/Dev-Main to Dev-Online-Plugins-v3 (//Tasks/UE4/Dev-Online-Plugins-v3)

Change 3045424 on 2016/07/11 by [REDACTED_PERSON]

	Add OnlineSubsystem plugins to Ocean.uproject

Change 3045478 on 2016/07/11 by [REDACTED_PERSON]

	Crash fix when creating a split screen player via DebugCreatePlayer 1
	null check UniqueId

Change 3046785 on 2016/07/12 by [REDACTED_PERSON]

	Move new hotfix files into plugin

Change 3047014 on 2016/07/12 by [REDACTED_PERSON]

	Fix for VOIP not working

Change 3047248 on 2016/07/12 by [REDACTED_PERSON]

	Fix build warning of member variable initialiization order.

Change 3047294 on 2016/07/12 by [REDACTED_PERSON]

	Fix for ensure on shutdown in Qos module

diff --git a/Engine/Build/JunkManifest.txt b/Engine/Build/JunkManifest.txt
--- a/Engine/Build/JunkManifest.txt
+++ b/Engine/Build/JunkManifest.txt
@@ -28,6 +28,7 @@ Engine\Binaries\Win64\*-SlateStyle*.*
 Engine\Binaries\Win64\*-CameraEditor*.*
 Engine\Binaries\Win64\*-CameraDebugger*.*
 Engine\Binaries\Win64\*-SessionLauncher*.*
+Engine\Binaries\Win64\*-OnlineSubsystem*.*
 Engine\Binaries\Mac\UnrealFrontend.app\Contents\Frameworks\*-CoreSlateStyle.*
 Engine\Binaries\Mac\UnrealFrontend.app\Contents\Frameworks\*-SlateStyle.*
 Engine\Binaries\Mac\UnrealFrontend.app\Contents\MacOS\*-XCodeSourceCodeAccess.*
diff --git a/Engine/Plugins/Editor/EpicSurvey/Source/EpicSurvey/EpicSurvey.Build.cs b/Engine/Plugins/Editor/EpicSurvey/Source/EpicSurvey/EpicSurvey.Build.cs
--- a/Engine/Plugins/Editor/EpicSurvey/Source/EpicSurvey/EpicSurvey.Build.cs
+++ b/Engine/Plugins/Editor/EpicSurvey/Source/EpicSurvey/EpicSurvey.Build.cs
@@ -16,7 +16,6 @@ namespace UnrealBuildTool.Rules
                     "InputCore",
 					"ImageWrapper",
 					"LevelEditor",
-					"OnlineSubSystem",
 					"Slate",
                     "EditorStyle",
 				}
@@ -27,7 +26,7 @@ namespace UnrealBuildTool.Rules
 				{
 					"HTTP",
 					"Json",
-					"OnlineSubSystem",
+					"OnlineSubsystem",
 					"SlateCore",
 					"UnrealEd",
 				}
@@ -46,16 +45,6 @@ namespace UnrealBuildTool.Rules
 					"MainFrame"
 				}
 			);
-
-			if (UEBuildConfiguration.bCompileMcpOSS == true)
-			{
-				DynamicallyLoadedModuleNames.AddRange(
-					new string[]
-					{
-						"OnlineSubsystemMcp",
-					}
-				);
-			}
 		}
 	}
 }
\ No newline at end of file
diff --git a/Engine/Plugins/Experimental/HTML5Networking/Source/HTML5Networking/HTML5Networking.Build.cs b/Engine/Plugins/Experimental/HTML5Networking/Source/HTML5Networking/HTML5Networking.Build.cs
--- a/Engine/Plugins/Experimental/HTML5Networking/Source/HTML5Networking/HTML5Networking.Build.cs
+++ b/Engine/Plugins/Experimental/HTML5Networking/Source/HTML5Networking/HTML5Networking.Build.cs
@@ -13,7 +13,6 @@ namespace UnrealBuildTool.Rules
                     "Engine",
                     "ImageCore",
                     "Sockets",
-                    "OnlineSubSystemUtils",
                     "libWebSockets",
                     "zlib"
                 }
diff --git a/Engine/Plugins/NetcodeUnitTest/NetcodeUnitTest/Source/NetcodeUnitTest/Private/Net/NUTUtilNet.cpp b/Engine/Plugins/NetcodeUnitTest/NetcodeUnitTest/Source/NetcodeUnitTest/Private/Net/NUTUtilNet.cpp
--- a/Engine/Plugins/NetcodeUnitTest/NetcodeUnitTest/Source/NetcodeUnitTest/Private/Net/NUTUtilNet.cpp
+++ b/Engine/Plugins/NetcodeUnitTest/NetcodeUnitTest/Source/NetcodeUnitTest/Private/Net/NUTUtilNet.cpp
@@ -16,6 +16,7 @@
 #include "DataChannel.h"
 #include "OnlineBeaconClient.h"
 #include "NetworkVersion.h"
+#include "OnlineSubsystemTypes.h"
 
 // Forward declarations
 class FWorldTickHook;
diff --git a/Engine/Plugins/Online/Android/OnlineSubsystemGooglePlay/Source/Private/OnlineAchievementsInterfaceGooglePlay.cpp b/Engine/Plugins/Online/Android/OnlineSubsystemGooglePlay/Source/Private/OnlineAchievementsInterfaceGooglePlay.cpp
new file mode 100644
--- /dev/null
+++ b/Engine/Plugins/Online/Android/OnlineSubsystemGooglePlay/Source/Private/OnlineAchievementsInterfaceGooglePlay.cpp
@@ -0,0 +1,332 @@
+// Copyright 1998-2016 Epic Games, Inc. All Rights Reserved.
+
+#include "OnlineSubsystemGooglePlayPrivatePCH.h"
+#include "OnlineAsyncTaskManagerGooglePlay.h"
+#include "Online.h"
+#include "Android/AndroidJNI.h"
+#include "OnlineAsyncTaskGooglePlayQueryAchievements.h"
+
+using namespace gpg;
+
+FOnlineAchievementsGooglePlay::FOnlineAchievementsGooglePlay( FOnlineSubsystemGooglePlay* InSubsystem )
+	: AndroidSubsystem(InSubsystem)
+{
+	// Make sure the Google achievement cache is initialized to an invalid status so we know to query it first
+	GoogleAchievements.status = ResponseStatus::ERROR_TIMEOUT;
+}
+
+void FOnlineAchievementsGooglePlay::QueryAchievements(const FUniqueNetId& PlayerId, const FOnQueryAchievementsCompleteDelegate& Delegate)
+{
+	if (AndroidSubsystem->GetGameServices() == nullptr)
+	{
+		Delegate.ExecuteIfBound(PlayerId, false);
+		return;
+	}
+
+	auto QueryTask = new FOnlineAsyncTaskGooglePlayQueryAchievements(AndroidSubsystem, FUniqueNetIdString(PlayerId), Delegate);
+	AndroidSubsystem->QueueAsyncTask(QueryTask);
+}
+
+
+void FOnlineAchievementsGooglePlay::QueryAchievementDescriptions(const FUniqueNetId& PlayerId, const FOnQueryAchievementsCompleteDelegate& Delegate)
+{
+	// Just query achievements to get descriptions
+	// FIXME: This feels a little redundant, but we can see how platforms evolve, and make a decision then
+	QueryAchievements( PlayerId, Delegate );
+}
+
+void FOnlineAchievementsGooglePlay::FinishAchievementWrite(
+	const FUniqueNetId& PlayerId,
+	const bool bWasSuccessful,
+	FOnlineAchievementsWriteRef WriteObject,
+	FOnAchievementsWrittenDelegate Delegate)
+{
+	// If the Google achievement cache is invalid, we can't do anything here - need to
+	// be able to get achievement type and max steps for incremental achievements.
+	if (GoogleAchievements.status != ResponseStatus::VALID && GoogleAchievements.status != ResponseStatus::VALID_BUT_STALE)
+	{
+		Delegate.ExecuteIfBound(PlayerId, false);
+		return;
+	}
+
+	auto DefaultSettings = GetDefault<UAndroidRuntimeSettings>();
+
+	for (FStatPropertyArray::TConstIterator It(WriteObject->Properties); It; ++It)
+	{
+		// Access the stat and the value.
+		const FVariantData& Stat = It.Value();
+
+		// Create an achievement object which should be reported to the server.
+		const FString UnrealAchievementId(It.Key().ToString());
+		float PercentComplete = 0.0f;
+
+		// Setup the percentage complete with the value we are writing from the variant type
+		switch (Stat.GetType())
+		{
+		case EOnlineKeyValuePairDataType::Int32:
+			{
+				int32 Value;
+				Stat.GetValue(Value);
+				PercentComplete = (float)Value;
+				break;
+			}
+		case EOnlineKeyValuePairDataType::Float:
+			{
+				float Value;
+				Stat.GetValue(Value);
+				PercentComplete = Value;
+				break;
+			}
+		default:
+			{
+				UE_LOG(LogOnline, Error, TEXT("FOnlineAchievementsGooglePlay Trying to write an achievement with incompatible format. Not a float or int"));
+				break;
+			}
+		}
+
+		const auto GoogleAchievement = GetGoogleAchievementFromUnrealId(UnrealAchievementId);
+		
+		if (!GoogleAchievement.Valid())
+		{
+			continue;
+		}
+
+		const FString GoogleId(GoogleAchievement.Id().c_str());
+
+		UE_LOG(LogOnline, Log, TEXT("Writing achievement name: %s, Google id: %s, progress: %.0f"),
+			*UnrealAchievementId, *GoogleId, PercentComplete);
+		
+		switch(GoogleAchievement.Type())
+		{
+			case gpg::AchievementType::INCREMENTAL:
+			{
+				float StepFraction = (PercentComplete / 100.0f) * GoogleAchievement.TotalSteps();
+				int RoundedSteps =  FMath::RoundToInt(StepFraction);
+
+				if(RoundedSteps > 0)
+				{
+					UE_LOG(LogOnline, Log, TEXT("  Incremental: setting progress to %d"), RoundedSteps);
+					AndroidSubsystem->GetGameServices()->Achievements().SetStepsAtLeast(GoogleAchievement.Id(), RoundedSteps);
+				}
+				else
+				{
+					UE_LOG(LogOnline, Log, TEXT("  Incremental: not setting progress to %d"), RoundedSteps);
+				}
+				break;
+			}
+
+			case gpg::AchievementType::STANDARD:
+			{
+				// Standard achievements only unlock if the progress is at least 100%.
+				if (PercentComplete >= 100.0f)
+				{
+					UE_LOG(LogOnline, Log, TEXT("  Standard: unlocking"));
+					AndroidSubsystem->GetGameServices()->Achievements().Unlock(GoogleAchievement.Id());
+				}
+				break;
+			}
+		}
+	}
+
+	Delegate.ExecuteIfBound(PlayerId, true);
+}
+
+void FOnlineAchievementsGooglePlay::WriteAchievements( const FUniqueNetId& PlayerId, FOnlineAchievementsWriteRef& WriteObject, const FOnAchievementsWrittenDelegate& Delegate )
+{
+	// @todo: verify the PlayerId passed in is the currently signed in player
+		
+	if (AndroidSubsystem->GetGameServices() == nullptr)
+	{
+		Delegate.ExecuteIfBound(PlayerId, false);
+		return;
+	}
+
+	// We need valid data from Google first because we use the achievement type to convert Unreal progress
+	// (percentage) to Google progress (binary locked/unlocked or integer number of arbitrary steps).
+	// Kick off a query if we don't have valid data.
+	if (GoogleAchievements.status != ResponseStatus::VALID)
+	{
+		auto QueryTask = new FOnlineAsyncTaskGooglePlayQueryAchievements(AndroidSubsystem, FUniqueNetIdString(PlayerId),
+			FOnQueryAchievementsCompleteDelegate::CreateRaw(
+				this,
+				&FOnlineAchievementsGooglePlay::FinishAchievementWrite,
+				WriteObject,
+				Delegate));
+		AndroidSubsystem->QueueAsyncTask(QueryTask);
+	}
+	else
+	{
+		FinishAchievementWrite(PlayerId, true, WriteObject, Delegate);
+	}
+}
+
+EOnlineCachedResult::Type FOnlineAchievementsGooglePlay::GetCachedAchievements(const FUniqueNetId& PlayerId, TArray<FOnlineAchievement>& OutAchievements)
+{
+	// @todo: verify the PlayerId passed in is the currently signed in player
+
+	OutAchievements.Empty();
+
+	if (GoogleAchievements.status != ResponseStatus::VALID)
+	{
+		return EOnlineCachedResult::NotFound;
+	}
+
+	auto DefaultSettings = GetDefault<UAndroidRuntimeSettings>();
+
+	for (auto& CurrentAchievement : GoogleAchievements.data)
+	{
+		const auto UnrealAchievement = GetUnrealAchievementFromGoogleAchievement(DefaultSettings, CurrentAchievement);
+		if (UnrealAchievement.Id.IsEmpty())
+		{
+			// Empty id means the achievement wasn't found.
+			continue;
+		}
+
+		OutAchievements.Add(UnrealAchievement);
+	}
+
+	return EOnlineCachedResult::Success;
+}
+
+
+EOnlineCachedResult::Type FOnlineAchievementsGooglePlay::GetCachedAchievementDescription(const FString& AchievementId, FOnlineAchievementDesc& OutAchievementDesc)
+{
+	auto FoundAchievement = GetGoogleAchievementFromUnrealId(AchievementId);
+	if (!FoundAchievement.Valid())
+	{
+		return EOnlineCachedResult::NotFound;
+	}
+
+	OutAchievementDesc.Title = FText::FromString(FString(FoundAchievement.Name().c_str()));
+	OutAchievementDesc.LockedDesc = FText::FromString(FString(FoundAchievement.Description().c_str()));
+	OutAchievementDesc.UnlockedDesc = FText::FromString(FString(FoundAchievement.Description().c_str()));
+	OutAchievementDesc.bIsHidden = FoundAchievement.State() == AchievementState::HIDDEN;
+	OutAchievementDesc.UnlockTime = FDateTime(0); // @todo: this
+
+	return EOnlineCachedResult::Success;
+}
+
+#if !UE_BUILD_SHIPPING
+bool FOnlineAchievementsGooglePlay::ResetAchievements( const FUniqueNetId& PlayerId )
+{
+	UE_LOG(LogOnline, Log, TEXT("Resetting Google Play achievements."));
+
+	extern void AndroidThunkCpp_ResetAchievements();
+	AndroidThunkCpp_ResetAchievements();
+
+	return false;
+};
+#endif // !UE_BUILD_SHIPPING
+
+EOnlineCachedResult::Type FOnlineAchievementsGooglePlay::GetCachedAchievement( const FUniqueNetId& PlayerId, const FString& AchievementId, FOnlineAchievement& OutAchievement )
+{
+	// @todo: verify the PlayerId passed in is the currently signed in player
+
+	auto FoundAchievement = GetGoogleAchievementFromUnrealId(AchievementId);
+	if (!FoundAchievement.Valid())
+	{
+		return EOnlineCachedResult::NotFound;
+	}
+
+	OutAchievement.Id = AchievementId;
+	OutAchievement.Progress = GetProgressFromGoogleAchievement(FoundAchievement);
+
+	return EOnlineCachedResult::Success;
+}
+
+Achievement FOnlineAchievementsGooglePlay::GetGoogleAchievementFromUnrealId(const FString& UnrealId) const
+{
+	auto DefaultSettings = GetDefault<UAndroidRuntimeSettings>();
+
+	FString TargetGoogleId;
+
+	for(const auto& Mapping : DefaultSettings->AchievementMap)
+	{
+		if(Mapping.Name.Equals(UnrealId))
+		{
+			TargetGoogleId = Mapping.AchievementID;
+			break;
+		}
+	}
+
+	for (auto& CurrentAchievement : GoogleAchievements.data)
+	{
+		if (!CurrentAchievement.Valid())
+		{
+			continue;
+		}
+
+		const FString CurrentGoogleId = CurrentAchievement.Id().c_str();
+
+		if (CurrentGoogleId.Equals(TargetGoogleId))
+		{
+			return CurrentAchievement;
+		}
+	}
+
+	// Not found, return empty achievement, Valid() will return false for this one.
+	return Achievement();
+}
+
+FString FOnlineAchievementsGooglePlay::GetUnrealIdFromGoogleId(const UAndroidRuntimeSettings* Settings, const FString& GoogleId)
+{
+	for(const auto& Mapping : Settings->AchievementMap)
+	{
+		if(Mapping.AchievementID.Equals(GoogleId))
+		{
+			return Mapping.Name;
+		}
+	}
+
+	return FString();
+}
+
+void FOnlineAchievementsGooglePlay::UpdateCache(const AchievementManager::FetchAllResponse& Results)
+{
+	GoogleAchievements = Results;
+}
+
+void FOnlineAchievementsGooglePlay::ClearCache()
+{
+	GoogleAchievements.status = ResponseStatus::ERROR_TIMEOUT; // Is there a better error to use here?
+	GoogleAchievements.data.clear();
+}
+
+double FOnlineAchievementsGooglePlay::GetProgressFromGoogleAchievement(const Achievement& InAchievement)
+{
+	if (!InAchievement.Valid())
+	{
+		return 0.0f;
+	}
+
+	if (InAchievement.State() == AchievementState::UNLOCKED)
+	{
+		return 100.0;
+	}
+
+	if (InAchievement.Type() == AchievementType::INCREMENTAL)
+	{
+		double Fraction = (double)InAchievement.CurrentSteps() / (double)InAchievement.TotalSteps();
+		return Fraction * 100.0;
+	}
+	
+	return 0.0;
+}
+
+FOnlineAchievement FOnlineAchievementsGooglePlay::GetUnrealAchievementFromGoogleAchievement(const UAndroidRuntimeSettings* Settings, const gpg::Achievement& GoogleAchievement)
+{
+	FOnlineAchievement OutAchievement;
+
+	if (!GoogleAchievement.Valid())
+	{
+		OutAchievement.Progress = 0.0;
+		return OutAchievement;
+	}
+
+	const FString GoogleId = GoogleAchievement.Id().c_str();
+
+	OutAchievement.Id = GetUnrealIdFromGoogleId(Settings, GoogleId);
+	OutAchievement.Progress = GetProgressFromGoogleAchievement(GoogleAchievement);
+
+	return OutAchievement;
+}
diff --git a/Engine/Plugins/Online/Android/OnlineSubsystemGooglePlay/Source/Private/OnlineAchievementsInterfaceGooglePlay.h b/Engine/Plugins/Online/Android/OnlineSubsystemGooglePlay/Source/Private/OnlineAchievementsInterfaceGooglePlay.h
new file mode 100644
--- /dev/null
+++ b/Engine/Plugins/Online/Android/OnlineSubsystemGooglePlay/Source/Private/OnlineAchievementsInterfaceGooglePlay.h
@@ -0,0 +1,91 @@
+// Copyright 1998-2016 Epic Games, Inc. All Rights Reserved.
+
+#pragma once
+
+#include "OnlineAchievementsInterface.h"
+#include "OnlineSubsystemTypes.h"
+#include "AndroidRuntimeSettings.h"
+#include <jni.h>
+
+#include "gpg/achievement_manager.h"
+
+/**
+ *	IOnlineAchievements - Interface class for Achievements
+ */
+class FOnlineAchievementsGooglePlay : public IOnlineAchievements
+{
+private:
+
+	/** Reference to the main GameCenter subsystem */
+	class FOnlineSubsystemGooglePlay* AndroidSubsystem;
+	
+	/** hide the default constructor, we need a reference to our OSS */
+	FOnlineAchievementsGooglePlay() {}
+
+	/** Our own cache of achievement data directly from Google Play */
+	gpg::AchievementManager::FetchAllResponse GoogleAchievements;
+
+	/** Looks up the Google achievement id in the mapping and returns the corresponding cached Achievement object */
+	gpg::Achievement GetGoogleAchievementFromUnrealId(const FString& UnrealId) const;
+
+	/** Uses the mapping to convert a Google achievement id to an Unreal achievement name */
+	static FString GetUnrealIdFromGoogleId(const UAndroidRuntimeSettings* Settings, const FString& GoogleId);
+
+	/** Convert the progress of a Google achievement to a double percentage 0.0 - 100.0 */
+	static double GetProgressFromGoogleAchievement(const gpg::Achievement& InAchievement);
+
+	/**
+	 * Convenience function to create an Unreal achievement from a Google achievement
+	 *
+	 * @param Settings the Android runtime settings containing the achievement id mapping
+	 * @param GoogleAchievement the Google achievement object to convert
+	 */
+	static FOnlineAchievement GetUnrealAchievementFromGoogleAchievement(
+		const UAndroidRuntimeSettings* Settings,
+		const gpg::Achievement& GoogleAchievement);
+
+	/**
+	 * Using the WriteObject, fires off achievement progress calls to the Google backend. Non-blocking.
+	 * The GoogleAchievements list should be valid before this is called.
+	 *
+	 * @param PlayerId the id of the player who's making progress
+	 * @param bWasSuccessful whether a previous QueryAchievements call was successful
+	 * @param WriteObject achievement write object provided by the user
+	 * @param Delegate delegate to execute when the write operation is finished
+	 */
+	void FinishAchievementWrite(
+		const FUniqueNetId& PlayerId,
+		const bool bWasSuccessful,
+		FOnlineAchievementsWriteRef WriteObject,
+		FOnAchievementsWrittenDelegate Delegate);
+	
+PACKAGE_SCOPE:
+	/** Clears the cache of Google achievements that was populated by a QueryAchievements() call. */
+	void ClearCache();
+
+	/** Called from the query achievements task to fill in the cache. */
+	void UpdateCache(const gpg::AchievementManager::FetchAllResponse& Results);
+
+public:
+
+	//~ Begin IOnlineAchievements Interface
+	virtual void WriteAchievements(const FUniqueNetId& PlayerId, FOnlineAchievementsWriteRef& WriteObject, const FOnAchievementsWrittenDelegate& Delegate = FOnAchievementsWrittenDelegate()) override;
+	virtual void QueryAchievements(const FUniqueNetId& PlayerId, const FOnQueryAchievementsCompleteDelegate& Delegate = FOnQueryAchievementsCompleteDelegate()) override;
+	virtual void QueryAchievementDescriptions(const FUniqueNetId& PlayerId, const FOnQueryAchievementsCompleteDelegate& Delegate = FOnQueryAchievementsCompleteDelegate()) override ;
+	virtual EOnlineCachedResult::Type GetCachedAchievement(const FUniqueNetId& PlayerId, const FString& AchievementId, FOnlineAchievement& OutAchievement) override;
+	virtual EOnlineCachedResult::Type GetCachedAchievements(const FUniqueNetId& PlayerId, TArray<FOnlineAchievement>& OutAchievements) override;
+	virtual EOnlineCachedResult::Type GetCachedAchievementDescription(const FString& AchievementId, FOnlineAchievementDesc& OutAchievementDesc) override;
+#if !UE_BUILD_SHIPPING
+	virtual bool ResetAchievements( const FUniqueNetId& PlayerId ) override;
+#endif // !UE_BUILD_SHIPPING
+	//~ End IOnlineAchievements Interface
+
+	/**
+	 * Constructor
+	 *
+	 * @param InSubsystem - A reference to the owning subsystem
+	 */
+	FOnlineAchievementsGooglePlay( class FOnlineSubsystemGooglePlay* InSubsystem );
+};
+
+typedef TSharedPtr<FOnlineAchievementsGooglePlay, ESPMode::ThreadSafe> FOnlineAchievementsGooglePlayPtr;
diff --git a/Engine/Plugins/Online/Android/OnlineSubsystemGooglePlay/Source/Private/OnlineAsyncTaskGooglePlayLogin.cpp b/Engine/Plugins/Online/Android/OnlineSubsystemGooglePlay/Source/Private/OnlineAsyncTaskGooglePlayLogin.cpp
new file mode 100644
--- /dev/null
+++ b/Engine/Plugins/Online/Android/OnlineSubsystemGooglePlay/Source/Private/OnlineAsyncTaskGooglePlayLogin.cpp
@@ -0,0 +1,76 @@
+// Copyright 1998-2016 Epic Games, Inc. All Rights Reserved.
+
+#include "OnlineSubsystemGooglePlayPrivatePCH.h"
+#include "OnlineAsyncTaskGooglePlayLogin.h"
+
+#include "gpg/builder.h"
+#include "gpg/debug.h"
+#include "gpg/types.h"
+
+FOnlineAsyncTaskGooglePlayLogin::FOnlineAsyncTaskGooglePlayLogin(
+	FOnlineSubsystemGooglePlay* InSubsystem,
+	int InPlayerId,
+	const FOnCompletedDelegate& InDelegate)
+	: FOnlineAsyncTaskGooglePlayAuthAction(InSubsystem)
+	, PlayerId(InPlayerId)
+	, Status(gpg::AuthStatus::ERROR_NOT_AUTHORIZED)
+	, Delegate(InDelegate)
+{
+}
+
+void FOnlineAsyncTaskGooglePlayLogin::Start_OnTaskThread()
+{
+	// If we haven't created a GameServices object yet, do so.
+	if (Subsystem->GameServicesPtr.get() == nullptr)
+	{
+		// Store the Subsystem pointer locally so that the OnAuthActionFinished lambda can capture it
+		FOnlineSubsystemGooglePlay* LocalSubsystem = Subsystem;
+
+		Subsystem->GameServicesPtr = gpg::GameServices::Builder()
+			.SetDefaultOnLog(gpg::LogLevel::VERBOSE)
+			.SetOnAuthActionStarted([](gpg::AuthOperation Op) {
+				UE_LOG(LogOnline, Log, TEXT("GPG OnAuthActionStarted: %s"), *FString(DebugString(Op).c_str()));
+			})
+			.SetOnAuthActionFinished([LocalSubsystem](gpg::AuthOperation Op, gpg::AuthStatus LocalStatus) {
+				UE_LOG(LogOnline, Log, TEXT("GPG OnAuthActionFinished: %s, AuthStatus: %s"),
+					*FString(DebugString(Op).c_str()),
+					*FString(DebugString(LocalStatus).c_str()));
+				LocalSubsystem->OnAuthActionFinished(Op, LocalStatus);
+			})
+			.Create(Subsystem->PlatformConfiguration);
+	}
+	else if(Subsystem->GameServicesPtr->IsAuthorized())
+	{
+		// We have a GameServices object and the user is authorized, nothing else to do.
+		Status = gpg::AuthStatus::VALID;
+		bWasSuccessful = true;
+		bIsComplete = true;
+	}
+	else
+	{
+		// We have created the GameServices object but the user isn't authorized.
+		bWasSuccessful = false;
+		bIsComplete = true;
+	}
+}
+
+void FOnlineAsyncTaskGooglePlayLogin::Finalize()
+{
+	// Async task manager owns the task and is responsible for cleaning it up.
+	Subsystem->CurrentLoginTask = nullptr;
+}
+
+void FOnlineAsyncTaskGooglePlayLogin::TriggerDelegates()
+{
+	Delegate.ExecuteIfBound();
+}
+
+void FOnlineAsyncTaskGooglePlayLogin::OnAuthActionFinished(gpg::AuthOperation InOp, gpg::AuthStatus InStatus)
+{
+	if (InOp == gpg::AuthOperation::SIGN_IN)
+	{
+		Status = InStatus;
+		bWasSuccessful = Status == gpg::AuthStatus::VALID;
+		bIsComplete = true;
+	}
+}
diff --git a/Engine/Plugins/Online/Android/OnlineSubsystemGooglePlay/Source/Private/OnlineAsyncTaskGooglePlayLogin.h b/Engine/Plugins/Online/Android/OnlineSubsystemGooglePlay/Source/Private/OnlineAsyncTaskGooglePlayLogin.h
new file mode 100644
--- /dev/null
+++ b/Engine/Plugins/Online/Android/OnlineSubsystemGooglePlay/Source/Private/OnlineAsyncTaskGooglePlayLogin.h
@@ -0,0 +1,44 @@
+// Copyright 1998-2016 Epic Games, Inc. All Rights Reserved.
+
+#pragma once
+
+#include "OnlineAsyncTaskManager.h"
+#include "OnlineAsyncTaskGooglePlayAuthAction.h"
+#include "OnlineSubsystemGooglePlayPackage.h"
+
+#include "gpg/status.h"
+#include "gpg/types.h"
+
+class FOnlineSubsystemGooglePlay;
+
+class FOnlineAsyncTaskGooglePlayLogin : public FOnlineAsyncTaskGooglePlayAuthAction
+{
+public:
+	/** Delegate fired upon completion. */
+	DECLARE_DELEGATE(FOnCompletedDelegate);
+
+	/**
+	 * Constructor.
+	 *
+	 * @param InSubsystem a pointer to the owning subsysetm
+	 * @param InPlayerId index of the player who's logging in
+	 */
+	FOnlineAsyncTaskGooglePlayLogin(FOnlineSubsystemGooglePlay* InSubsystem, int InPlayerId, const FOnCompletedDelegate& InDelegate);
+
+	// FOnlineAsyncItem
+	virtual FString ToString() const override { return TEXT("Login"); }
+	virtual void Finalize() override;
+	virtual void TriggerDelegates() override;
+
+private:
+	/** The subsystem is the only class that should be calling OnAuthActionFinished */
+	friend class FOnlineSubsystemGooglePlay;
+
+	// FOnlineAsyncTaskGooglePlayAuthAction
+	virtual void OnAuthActionFinished(gpg::AuthOperation InOp, gpg::AuthStatus InStatus) override;
+	virtual void Start_OnTaskThread() override;
+
+	int PlayerId;
+	gpg::AuthStatus Status;
+	FOnCompletedDelegate Delegate;
+};
diff --git a/Engine/Plugins/Online/Android/OnlineSubsystemGooglePlay/Source/Private/OnlineAsyncTaskGooglePlayQueryAchievements.cpp b/Engine/Plugins/Online/Android/OnlineSubsystemGooglePlay/Source/Private/OnlineAsyncTaskGooglePlayQueryAchievements.cpp
new file mode 100644
--- /dev/null
+++ b/Engine/Plugins/Online/Android/OnlineSubsystemGooglePlay/Source/Private/OnlineAsyncTaskGooglePlayQueryAchievements.cpp
@@ -0,0 +1,44 @@
+// Copyright 1998-2016 Epic Games, Inc. All Rights Reserved.
+
+#include "OnlineSubsystemGooglePlayPrivatePCH.h"
+#include "OnlineAsyncTaskGooglePlayQueryAchievements.h"
+
+#include "gpg/achievement_manager.h"
+
+FOnlineAsyncTaskGooglePlayQueryAchievements::FOnlineAsyncTaskGooglePlayQueryAchievements(
+	FOnlineSubsystemGooglePlay* InSubsystem,
+	const FUniqueNetIdString& InUserId,
+	const FOnQueryAchievementsCompleteDelegate& InDelegate)
+	: FOnlineAsyncTaskBasic(InSubsystem)
+	, UserId(InUserId)
+	, Delegate(InDelegate)
+{
+}
+
+void FOnlineAsyncTaskGooglePlayQueryAchievements::Finalize()
+{
+	if (bWasSuccessful)
+	{
+		Subsystem->GetAchievementsGooglePlay()->UpdateCache(Response);
+	}
+	else
+	{
+		Subsystem->GetAchievementsGooglePlay()->ClearCache();
+	}
+}
+
+void FOnlineAsyncTaskGooglePlayQueryAchievements::TriggerDelegates()
+{
+	Delegate.ExecuteIfBound(UserId, bWasSuccessful);
+}
+
+void FOnlineAsyncTaskGooglePlayQueryAchievements::Tick()
+{
+	// We're already running on the online thread. Using the blocking version of the API function
+	// won't block the main thread and simplifies things a bit.
+	// Try a 10 second timeout.
+	Response = Subsystem->GetGameServices()->Achievements().FetchAllBlocking(gpg::DataSource::CACHE_OR_NETWORK, gpg::Timeout(10000));
+	
+	bWasSuccessful = Response.status == gpg::ResponseStatus::VALID || Response.status == gpg::ResponseStatus::VALID_BUT_STALE;
+	bIsComplete = true;
+}
diff --git a/Engine/Plugins/Online/Android/OnlineSubsystemGooglePlay/Source/Private/OnlineAsyncTaskGooglePlayReadLeaderboard.cpp b/Engine/Plugins/Online/Android/OnlineSubsystemGooglePlay/Source/Private/OnlineAsyncTaskGooglePlayReadLeaderboard.cpp
new file mode 100644
--- /dev/null
+++ b/Engine/Plugins/Online/Android/OnlineSubsystemGooglePlay/Source/Private/OnlineAsyncTaskGooglePlayReadLeaderboard.cpp
@@ -0,0 +1,91 @@
+// Copyright 1998-2016 Epic Games, Inc. All Rights Reserved.
+
+#include "OnlineSubsystemGooglePlayPrivatePCH.h"
+#include "OnlineAsyncTaskGooglePlayReadLeaderboard.h"
+
+#include "gpg/leaderboard_manager.h"
+
+using namespace gpg;
+
+FOnlineAsyncTaskGooglePlayReadLeaderboard::FOnlineAsyncTaskGooglePlayReadLeaderboard(
+	FOnlineSubsystemGooglePlay* InSubsystem,
+	FOnlineLeaderboardReadRef& InReadObject,
+	const FString& InLeaderboardId)
+	: FOnlineAsyncTaskBasic(InSubsystem)
+	, ReadObject(InReadObject)
+	, LeaderboardId(InLeaderboardId)
+{
+	Response.status = ResponseStatus::ERROR_TIMEOUT;
+}
+
+void FOnlineAsyncTaskGooglePlayReadLeaderboard::Finalize()
+{
+	ReadObject->ReadState = bWasSuccessful ? EOnlineAsyncTaskState::Done : EOnlineAsyncTaskState::Failed;
+}
+
+void FOnlineAsyncTaskGooglePlayReadLeaderboard::TriggerDelegates()
+{
+	Subsystem->GetLeaderboardsInterface()->TriggerOnLeaderboardReadCompleteDelegates(bWasSuccessful);
+}
+
+void FOnlineAsyncTaskGooglePlayReadLeaderboard::Tick()
+{
+	// Convert our FString leaderboard ID to ASCII for the API's std::string.
+	auto ConvertedId = FOnlineSubsystemGooglePlay::ConvertFStringToStdString(LeaderboardId);
+	Response = Subsystem->GetGameServices()->Leaderboards().FetchScoreSummaryBlocking(
+		DataSource::CACHE_OR_NETWORK,
+		Timeout(10000),
+		ConvertedId,
+		LeaderboardTimeSpan::ALL_TIME,
+		LeaderboardCollection::PUBLIC);
+
+	bWasSuccessful = false;
+
+	if (Response.status != ResponseStatus::VALID)
+	{
+		bIsComplete = true;
+		return;
+	}
+
+	// We can only get the current user's leaderboard score from Google Play, so just add one row with it.
+	//TSharedRef<const FUniqueNetId> UserId = MakeShareable(new FUniqueNetIdString(FString(TEXT("0"))));
+	auto UserId = Subsystem->GetIdentityGooglePlay()->GetCurrentUserId();
+	if (!UserId.IsValid())
+	{
+		// If there's no user signed in, can't get leaderboard.
+		bIsComplete = true;
+		return;
+	}
+
+	FOnlineStatsRow* UserRow = ReadObject.Get().FindPlayerRecord(*UserId);
+	if (UserRow == NULL)
+	{
+		UserRow = new (ReadObject->Rows) FOnlineStatsRow(UserId->ToString(), UserId.ToSharedRef());
+	}
+
+	for (int32 StatIdx = 0; StatIdx < ReadObject->ColumnMetadata.Num(); StatIdx++)
+	{
+		const FColumnMetaData& ColumnMeta = ReadObject->ColumnMetadata[StatIdx];
+
+		FVariantData* LastColumn = NULL;
+		switch (ColumnMeta.DataType)
+		{
+			case EOnlineKeyValuePairDataType::Int32:
+			{
+				int32 Value = Response.data.CurrentPlayerScore().Value();
+				LastColumn = &(UserRow->Columns.Add(ColumnMeta.ColumnName, FVariantData(Value)));
+				bWasSuccessful = true;
+				break;
+			}
+
+			default:
+			{
+				UE_LOG_ONLINE(Warning, TEXT("Unsupported key value pair during retrieval from Google Play %s"), ColumnMeta.ColumnName);
+				break;
+			}
+		}
+	}
+	
+	bIsComplete = true;
+}
+
diff --git a/Engine/Plugins/Online/Android/OnlineSubsystemGooglePlay/Source/Private/OnlineAsyncTaskGooglePlayReadLeaderboard.h b/Engine/Plugins/Online/Android/OnlineSubsystemGooglePlay/Source/Private/OnlineAsyncTaskGooglePlayReadLeaderboard.h
new file mode 100644
--- /dev/null
+++ b/Engine/Plugins/Online/Android/OnlineSubsystemGooglePlay/Source/Private/OnlineAsyncTaskGooglePlayReadLeaderboard.h
@@ -0,0 +1,36 @@
+// Copyright 1998-2016 Epic Games, Inc. All Rights Reserved.
+
+#pragma once
+
+#include "OnlineAsyncTaskManager.h"
+
+#include "gpg/leaderboard_manager.h"
+
+class FOnlineSubsystemGooglePlay;
+
+class FOnlineAsyncTaskGooglePlayReadLeaderboard : public FOnlineAsyncTaskBasic<FOnlineSubsystemGooglePlay>
+{
+public:
+	FOnlineAsyncTaskGooglePlayReadLeaderboard(
+		FOnlineSubsystemGooglePlay* InSubsystem,
+		FOnlineLeaderboardReadRef& InReadObject,
+		const FString& InLeaderboardId);
+
+	// FOnlineAsyncItem
+	virtual FString ToString() const override { return TEXT("ReadLeaderboard"); }
+	virtual void Finalize() override;
+	virtual void TriggerDelegates() override;
+
+	// FOnlineAsyncTask
+	virtual void Tick() override;
+
+private:
+	/** Leaderboard read data */
+	FOnlineLeaderboardReadRef ReadObject;
+
+	/** Google Play leaderboard id */
+	FString LeaderboardId;
+
+	/** API query result */
+	gpg::LeaderboardManager::FetchScoreSummaryResponse Response;
+};
diff --git a/Engine/Plugins/Online/Android/OnlineSubsystemGooglePlay/Source/Private/OnlineAsyncTaskGooglePlayShowLoginUI.cpp b/Engine/Plugins/Online/Android/OnlineSubsystemGooglePlay/Source/Private/OnlineAsyncTaskGooglePlayShowLoginUI.cpp
new file mode 100644
--- /dev/null
+++ b/Engine/Plugins/Online/Android/OnlineSubsystemGooglePlay/Source/Private/OnlineAsyncTaskGooglePlayShowLoginUI.cpp
@@ -0,0 +1,90 @@
+// Copyright 1998-2016 Epic Games, Inc. All Rights Reserved.
+
+#include "OnlineSubsystemGooglePlayPrivatePCH.h"
+#include "OnlineAsyncTaskGooglePlayShowLoginUI.h"
+
+#include "gpg/player_manager.h"
+
+
+FOnlineAsyncTaskGooglePlayShowLoginUI::FOnlineAsyncTaskGooglePlayShowLoginUI(
+	FOnlineSubsystemGooglePlay* InSubsystem,
+	int InPlayerId,
+	const FOnLoginUIClosedDelegate& InDelegate)
+	: FOnlineAsyncTaskGooglePlayAuthAction(InSubsystem)
+	, PlayerId(InPlayerId)
+	, Delegate(InDelegate)
+{
+}
+
+void FOnlineAsyncTaskGooglePlayShowLoginUI::Start_OnTaskThread()
+{	
+	if(Subsystem->GetGameServices() == nullptr)
+	{
+		UE_LOG(LogOnline, Log, TEXT("FOnlineAsyncTaskGooglePlayShowLoginUI: GameServicesPtr is null."));
+		bWasSuccessful = false;
+		bIsComplete = true;
+		return;
+	}
+
+	if(Subsystem->GetGameServices()->IsAuthorized())
+	{
+		// User is already authorized, nothing to do.
+		bWasSuccessful = true;
+		Subsystem->GetGameServices()->Players().FetchSelf([this](gpg::PlayerManager::FetchSelfResponse const &response) { OnFetchSelfResponse(response); });
+	}
+
+	// The user isn't authorized, show the sign-in UI.
+	Subsystem->GetGameServices()->StartAuthorizationUI();
+}
+
+void FOnlineAsyncTaskGooglePlayShowLoginUI::Finalize()
+{
+	// Async task manager owns the task and is responsible for cleaning it up.
+	Subsystem->CurrentShowLoginUITask = nullptr;
+}
+
+void FOnlineAsyncTaskGooglePlayShowLoginUI::TriggerDelegates()
+{
+	TSharedPtr<const FUniqueNetIdString> UserId = Subsystem->GetIdentityGooglePlay()->GetCurrentUserId();
+
+	if (bWasSuccessful && !UserId.IsValid())
+	{
+		UserId = MakeShareable(new FUniqueNetIdString());
+		Subsystem->GetIdentityGooglePlay()->SetCurrentUserId(UserId);
+	}
+	else if (!bWasSuccessful)
+	{
+		Subsystem->GetIdentityGooglePlay()->SetCurrentUserId(nullptr);
+	}
+
+	Delegate.ExecuteIfBound(UserId, PlayerId);
+}
+
+void FOnlineAsyncTaskGooglePlayShowLoginUI::OnAuthActionFinished(gpg::AuthOperation InOp, gpg::AuthStatus InStatus)
+{
+	if (InOp == gpg::AuthOperation::SIGN_IN)
+	{
+		bWasSuccessful = InStatus == gpg::AuthStatus::VALID;
+		if(bWasSuccessful)
+		{
+			Subsystem->GetGameServices()->Players().FetchSelf([this](gpg::PlayerManager::FetchSelfResponse const &response) { OnFetchSelfResponse(response); });
+		}
+		else
+		{
+			bIsComplete = true;
+		}
+	}
+}
+
+void FOnlineAsyncTaskGooglePlayShowLoginUI::OnFetchSelfResponse(const gpg::PlayerManager::FetchSelfResponse& SelfResponse)
+{
+	if(gpg::IsSuccess(SelfResponse.status))
+	{
+		Subsystem->GetIdentityGooglePlay()->SetPlayerDataFromFetchSelfResponse(SelfResponse.data);
+	}
+	else
+	{
+		UE_LOG(LogOnline, Warning, TEXT("FetchSelf Response Status Not Successful"));
+	}
+	bIsComplete = true;
+}
diff --git a/Engine/Plugins/Online/Android/OnlineSubsystemGooglePlay/Source/Private/OnlineAsyncTaskGooglePlayShowLoginUI.h b/Engine/Plugins/Online/Android/OnlineSubsystemGooglePlay/Source/Private/OnlineAsyncTaskGooglePlayShowLoginUI.h
new file mode 100644
--- /dev/null
+++ b/Engine/Plugins/Online/Android/OnlineSubsystemGooglePlay/Source/Private/OnlineAsyncTaskGooglePlayShowLoginUI.h
@@ -0,0 +1,43 @@
+// Copyright 1998-2016 Epic Games, Inc. All Rights Reserved.
+
+#pragma once
+
+#include "OnlineAsyncTaskManager.h"
+#include "OnlineExternalUIInterface.h"
+#include "OnlineAsyncTaskGooglePlayAuthAction.h"
+#include "OnlineSubsystemGooglePlayPackage.h"
+
+#include "gpg/types.h"
+#include "gpg/player_manager.h"
+
+
+class FOnlineSubsystemGooglePlay;
+
+class FOnlineAsyncTaskGooglePlayShowLoginUI : public FOnlineAsyncTaskGooglePlayAuthAction
+{
+public:
+	/**
+	 * Constructor.
+	 *
+	 * @param InSubsystem a pointer to the owning subsystem
+	 * @param InPlayerId index of the player who's logging in
+	 */
+	FOnlineAsyncTaskGooglePlayShowLoginUI(FOnlineSubsystemGooglePlay* InSubsystem, int InPlayerId, const FOnLoginUIClosedDelegate& InDelegate);
+
+	// FOnlineAsyncItem
+	virtual FString ToString() const override { return TEXT("ShowLoginUI"); }
+	virtual void Finalize() override;
+	virtual void TriggerDelegates() override;
+
+private:
+	/** The subsystem is the only class that should be calling OnAuthActionFinished */
+	friend class FOnlineSubsystemGooglePlay;
+
+	// FOnlineAsyncTaskGooglePlayAuthAction
+	virtual void OnAuthActionFinished(gpg::AuthOperation InOp, gpg::AuthStatus InStatus) override;
+	virtual void OnFetchSelfResponse(const gpg::PlayerManager::FetchSelfResponse& SelfResponse);
+	virtual void Start_OnTaskThread() override;
+
+	int PlayerId;
+	FOnLoginUIClosedDelegate Delegate;
+};
diff --git a/Engine/Plugins/Online/Android/OnlineSubsystemGooglePlay/Source/Private/OnlineExternalUIInterfaceGooglePlay.cpp b/Engine/Plugins/Online/Android/OnlineSubsystemGooglePlay/Source/Private/OnlineExternalUIInterfaceGooglePlay.cpp
new file mode 100644
--- /dev/null
+++ b/Engine/Plugins/Online/Android/OnlineSubsystemGooglePlay/Source/Private/OnlineExternalUIInterfaceGooglePlay.cpp
@@ -0,0 +1,104 @@
+// Copyright 1998-2016 Epic Games, Inc. All Rights Reserved.
+
+#include "OnlineSubsystemGooglePlayPrivatePCH.h"
+#include "AndroidRuntimeSettings.h"
+#include "CoreUObject.h"
+
+#include "gpg/achievement_manager.h"
+#include "gpg/leaderboard_manager.h"
+
+FOnlineExternalUIGooglePlay::FOnlineExternalUIGooglePlay(FOnlineSubsystemGooglePlay* InSubsystem)
+	: Subsystem(InSubsystem)
+{
+	check(Subsystem != nullptr);
+}
+
+bool FOnlineExternalUIGooglePlay::ShowLoginUI(const int ControllerIndex, bool bShowOnlineOnly, const FOnLoginUIClosedDelegate& Delegate)
+{
+	Subsystem->StartShowLoginUITask(ControllerIndex, Delegate);	
+	return true;
+}
+
+bool FOnlineExternalUIGooglePlay::ShowFriendsUI(int32 LocalUserNum)
+{
+	return false;
+}
+
+bool FOnlineExternalUIGooglePlay::ShowInviteUI(int32 LocalUserNum, FName SessionMame)
+{
+	return false;
+}
+
+bool FOnlineExternalUIGooglePlay::ShowAchievementsUI(int32 LocalUserNum) 
+{
+	if (Subsystem->GetGameServices() == nullptr)
+	{
+		return false;
+	}
+
+	if (!Subsystem->GetGameServices()->IsAuthorized())
+	{
+		return false;
+	}
+
+	Subsystem->GetGameServices()->Achievements().ShowAllUI();
+	return true;
+}
+
+bool FOnlineExternalUIGooglePlay::ShowLeaderboardUI(const FString& LeaderboardName)
+{
+	if (Subsystem->GetGameServices() == nullptr)
+	{
+		return false;
+	}
+
+	if (!Subsystem->GetGameServices()->IsAuthorized())
+	{
+		return false;
+	}
+
+	auto DefaultSettings = GetDefault<UAndroidRuntimeSettings>();
+
+	for(const auto& Mapping : DefaultSettings->LeaderboardMap)
+	{
+		if(Mapping.Name == LeaderboardName)
+		{
+			auto ConvertedId = FOnlineSubsystemGooglePlay::ConvertFStringToStdString(Mapping.LeaderboardID);
+			Subsystem->GetGameServices()->Leaderboards().ShowUI(ConvertedId);
+
+			return true;
+		}
+	}
+	
+	return false;
+}
+
+bool FOnlineExternalUIGooglePlay::ShowWebURL(const FString& Url, const FShowWebUrlParams& ShowParams, const FOnShowWebUrlClosedDelegate& Delegate)
+{
+	return false;
+}
+
+bool FOnlineExternalUIGooglePlay::CloseWebURL()
+{
+	return false;
+}
+
+bool FOnlineExternalUIGooglePlay::ShowProfileUI( const FUniqueNetId& Requestor, const FUniqueNetId& Requestee, const FOnProfileUIClosedDelegate& Delegate )
+{
+	return false;
+}
+
+bool FOnlineExternalUIGooglePlay::ShowAccountUpgradeUI(const FUniqueNetId& UniqueId)
+{
+	return false;
+}
+
+bool FOnlineExternalUIGooglePlay::ShowStoreUI(int32 LocalUserNum, const FShowStoreParams& ShowParams, const FOnShowStoreUIClosedDelegate& Delegate)
+{
+	return false;
+}
+
+bool FOnlineExternalUIGooglePlay::ShowSendMessageUI(int32 LocalUserNum, const FShowSendMessageParams& ShowParams, const FOnShowSendMessageUIClosedDelegate& Delegate)
+{
+	return false;
+}
\ No newline at end of file
diff --git a/Engine/Plugins/Online/Android/OnlineSubsystemGooglePlay/Source/Private/OnlineIdentityInterfaceGooglePlay.h b/Engine/Plugins/Online/Android/OnlineSubsystemGooglePlay/Source/Private/OnlineIdentityInterfaceGooglePlay.h
new file mode 100644
--- /dev/null
+++ b/Engine/Plugins/Online/Android/OnlineSubsystemGooglePlay/Source/Private/OnlineIdentityInterfaceGooglePlay.h
@@ -0,0 +1,89 @@
+// Copyright 1998-2016 Epic Games, Inc. All Rights Reserved.
+
+#pragma once
+
+#include "OnlineIdentityInterface.h"
+#include "OnlineSubsystemGooglePlayPackage.h"
+
+#include "gpg/status.h"
+#include "gpg/player.h"
+
+class FOnlineIdentityGooglePlay :
+	public IOnlineIdentity
+{
+private:
+	
+	bool bPrevLoggedIn;
+	bool bLoggedIn;
+	FString PlayerAlias;
+	int32 CurrentLoginUserNum;
+	class FOnlineSubsystemGooglePlay* MainSubsystem;
+
+	bool bLoggingInUser;
+	bool bRegisteringUser;
+
+	/** UID for this identity */
+	TSharedPtr< const FUniqueNetIdString > UniqueNetId;
+
+	struct FPendingConnection
+	{
+		FOnlineIdentityGooglePlay * ConnectionInterface;
+		bool IsConnectionPending;
+	};
+
+	/** Instance of the connection context data */
+	static FPendingConnection PendingConnectRequest;
+
+PACKAGE_SCOPE:
+
+	/**
+	 * Default Constructor
+	 */
+	FOnlineIdentityGooglePlay(FOnlineSubsystemGooglePlay* InSubsystem);
+
+	/** Allow individual interfaces to access the currently signed-in user's id */
+	TSharedPtr<const FUniqueNetIdString> GetCurrentUserId() const { return UniqueNetId; }
+
+	void SetCurrentUserId(TSharedPtr<const FUniqueNetIdString> InUniqueNetId) { UniqueNetId = InUniqueNetId; }
+
+	/** Called from ExternalUIInterface to set UniqueId and PlayerAlias after authentication */
+	void SetPlayerDataFromFetchSelfResponse(const gpg::Player& PlayerData);
+
+public:
+
+	//~ Begin IOnlineIdentity Interface
+	virtual TSharedPtr<FUserOnlineAccount> GetUserAccount(const FUniqueNetId& UserId) const override;
+	virtual TArray<TSharedPtr<FUserOnlineAccount> > GetAllUserAccounts() const override;
+	virtual bool Login(int32 LocalUserNum, const FOnlineAccountCredentials& AccountCredentials) override;
+	virtual bool Logout(int32 LocalUserNum) override;
+	virtual bool AutoLogin(int32 LocalUserNum) override;
+	virtual TSharedPtr<const FUniqueNetId> GetUniquePlayerId(int32 LocalUserNum) const override;
+	virtual TSharedPtr<const FUniqueNetId> CreateUniquePlayerId(uint8* Bytes, int32 Size) override;
+	virtual TSharedPtr<const FUniqueNetId> CreateUniquePlayerId(const FString& Str) override;
+	virtual ELoginStatus::Type GetLoginStatus(int32 LocalUserNum) const override;
+	virtual ELoginStatus::Type GetLoginStatus(const FUniqueNetId& UserId) const override;
+	virtual FString GetPlayerNickname(int32 LocalUserNum) const override;
+	virtual FString GetPlayerNickname(const FUniqueNetId& UserId) const override;
+	virtual FString GetAuthToken(int32 LocalUserNum) const override;
+	virtual void GetUserPrivilege(const FUniqueNetId& UserId, EUserPrivileges::Type Privilege, const FOnGetUserPrivilegeCompleteDelegate& Delegate) override;
+	virtual FPlatformUserId GetPlatformUserIdFromUniqueNetId(const FUniqueNetId& NetId) override;
+	virtual FString GetAuthType() const override;
+	//~ End IOnlineIdentity interface
+
+public:
+
+	//~ Begin IOnlineIdentity Interface
+	void Tick(float DeltaTime);
+	//~ End IOnlineIdentity Interface
+	
+	//platform specific
+	void OnLogin(const bool InLoggedIn,  const FString& InPlayerId, const FString& InPlayerAlias);
+	void OnLogout(const bool InLoggedIn);
+
+	void OnLoginCompleted(const int playerID, gpg::AuthStatus errorCode);
+
+};
+
+
+typedef TSharedPtr<FOnlineIdentityGooglePlay, ESPMode::ThreadSafe> FOnlineIdentityGooglePlayPtr;
+
diff --git a/Engine/Plugins/Online/Android/OnlineSubsystemGooglePlay/Source/Private/OnlineLeaderboardInterfaceGooglePlay.h b/Engine/Plugins/Online/Android/OnlineSubsystemGooglePlay/Source/Private/OnlineLeaderboardInterfaceGooglePlay.h
new file mode 100644
--- /dev/null
+++ b/Engine/Plugins/Online/Android/OnlineSubsystemGooglePlay/Source/Private/OnlineLeaderboardInterfaceGooglePlay.h
@@ -0,0 +1,45 @@
+// Copyright 1998-2016 Epic Games, Inc. All Rights Reserved.
+
+#pragma once
+
+#include "OnlineLeaderboardInterface.h"
+#include "AndroidRuntimeSettings.h"
+
+struct FOnlinePendingLeaderboardWrite
+{
+	FString LeaderboardName;
+	uint64 Score;
+};
+
+/**
+ * Interface definition for the online services leaderboard services 
+ */
+class FOnlineLeaderboardsGooglePlay : public IOnlineLeaderboards
+{
+public:
+	FOnlineLeaderboardsGooglePlay(class FOnlineSubsystemGooglePlay* InSubsystem);
+
+	virtual bool ReadLeaderboards(const TArray< TSharedRef<const FUniqueNetId> >& Players, FOnlineLeaderboardReadRef& ReadObject) override;
+	virtual bool ReadLeaderboardsForFriends(int32 LocalUserNum, FOnlineLeaderboardReadRef& ReadObject) override;
+	virtual void FreeStats(FOnlineLeaderboardRead& ReadObject) override;
+	virtual bool WriteLeaderboards(const FName& SessionName, const FUniqueNetId& Player, FOnlineLeaderboardWrite& WriteObject) override;
+	virtual bool FlushLeaderboards(const FName& SessionName) override;
+	virtual bool WriteOnlinePlayerRatings(const FName& SessionName, int32 LeaderboardId, const TArray<FOnlinePlayerScore>& PlayerScores) override;
+
+private:
+	/** Pointer to owning subsystem */
+	FOnlineSubsystemGooglePlay* Subsystem;
+
+	/**
+	 * Helper function to get the platform- and game-specific leaderboard ID from the JSON config file.
+	 *
+	 * @param LeaderboardName the cross-platform name of the leaderboard to look up
+	 * @return The unique ID for the leaderboard as specified in the config file.
+	*/
+	FString GetLeaderboardID(const FString& LeaderboardName);
+
+	/** Scores are cached here in WriteLeaderboards until FlushLeaderboards is called */
+	TArray<FOnlinePendingLeaderboardWrite> UnreportedScores;
+};
+
+typedef TSharedPtr<FOnlineLeaderboardsGooglePlay, ESPMode::ThreadSafe> FOnlineLeaderboardsGooglePlayPtr;
diff --git a/Engine/Plugins/Online/Android/OnlineSubsystemGooglePlay/Source/Private/OnlineStoreInterfaceGooglePlay.cpp b/Engine/Plugins/Online/Android/OnlineSubsystemGooglePlay/Source/Private/OnlineStoreInterfaceGooglePlay.cpp
new file mode 100644
--- /dev/null
+++ b/Engine/Plugins/Online/Android/OnlineSubsystemGooglePlay/Source/Private/OnlineStoreInterfaceGooglePlay.cpp
@@ -0,0 +1,365 @@
+// Copyright 1998-2016 Epic Games, Inc. All Rights Reserved.
+
+#include "OnlineSubsystemGooglePlayPrivatePCH.h"
+#include "OnlineAsyncTaskGooglePlayQueryInAppPurchases.h"
+#include "TaskGraphInterfaces.h"
+#include <jni.h>
+
+
+////////////////////////////////////////////////////////////////////
+/// FOnlineStoreGooglePlay implementation
+
+
+FOnlineStoreGooglePlay::FOnlineStoreGooglePlay(FOnlineSubsystemGooglePlay* InSubsystem)
+	: Subsystem( InSubsystem )
+	, CurrentQueryTask( nullptr )
+{
+	UE_LOG(LogOnline, Display, TEXT( "FOnlineStoreGooglePlay::FOnlineStoreGooglePlay" ));
+	
+	FString LicenseKey;
+	GConfig->GetString(TEXT("/Script/AndroidRuntimeSettings.AndroidRuntimeSettings"), TEXT("GooglePlayLicenseKey"), LicenseKey, GEngineIni);
+
+	extern void AndroidThunkCpp_Iap_SetupIapService(const FString&);
+	AndroidThunkCpp_Iap_SetupIapService(LicenseKey);
+}
+
+
+FOnlineStoreGooglePlay::~FOnlineStoreGooglePlay()
+{
+	UE_LOG(LogOnline, Display, TEXT( "FOnlineStoreGooglePlay::~FOnlineStoreGooglePlay" ));
+}
+
+
+bool FOnlineStoreGooglePlay::IsAllowedToMakePurchases()
+{
+	UE_LOG(LogOnline, Display, TEXT("FOnlineStoreGooglePlay::IsAllowedToMakePurchases"));
+
+	extern bool AndroidThunkCpp_Iap_IsAllowedToMakePurchases();
+	return AndroidThunkCpp_Iap_IsAllowedToMakePurchases();
+}
+
+
+bool FOnlineStoreGooglePlay::QueryForAvailablePurchases(const TArray<FString>& ProductIds, FOnlineProductInformationReadRef& InReadObject)
+{
+	UE_LOG(LogOnline, Display, TEXT( "FOnlineStoreGooglePlay::QueryForAvailablePurchases" ));
+	
+	ReadObject = InReadObject;
+	ReadObject->ReadState = EOnlineAsyncTaskState::InProgress;
+
+	TArray<bool> ConsumableFlags;
+	ConsumableFlags.AddZeroed(ProductIds.Num());
+
+	CurrentQueryTask = new FOnlineAsyncTaskGooglePlayQueryInAppPurchases(
+		Subsystem,
+		ProductIds,
+		ConsumableFlags);
+	Subsystem->QueueAsyncTask(CurrentQueryTask);
+
+	return true;
+}
+
+
+extern "C" void Java_com_epicgames_ue4_GooglePlayStoreHelper_nativeQueryComplete(JNIEnv* jenv, jobject thiz, jboolean bSuccess, jobjectArray productIDs, jobjectArray titles, jobjectArray descriptions, jobjectArray prices, jfloatArray pricesRaw, jobjectArray currencyCodes)
+{
+	TArray<FInAppPurchaseProductInfo> ProvidedProductInformation;
+
+	if (jenv && bSuccess)
+	{
+		jsize NumProducts = jenv->GetArrayLength(productIDs);
+		jsize NumTitles = jenv->GetArrayLength(titles);
+		jsize NumDescriptions = jenv->GetArrayLength(descriptions);
+		jsize NumPrices = jenv->GetArrayLength(prices);
+		jsize NumPricesRaw = jenv->GetArrayLength(pricesRaw);
+		jsize NumCurrencyCodes = jenv->GetArrayLength(currencyCodes);
+
+		ensure((NumProducts == NumTitles) && (NumProducts == NumDescriptions) && (NumProducts == NumPrices) && (NumProducts == NumPricesRaw) && (NumProducts == NumCurrencyCodes));
+
+		jfloat *PricesRaw = jenv->GetFloatArrayElements(pricesRaw, 0);
+
+		for (jsize Idx = 0; Idx < NumProducts; Idx++)
+		{
+			// Build the product information strings.
+
+			FInAppPurchaseProductInfo NewProductInfo;
+
+			jstring NextId = (jstring)jenv->GetObjectArrayElement(productIDs, Idx);
+			const char* charsId = jenv->GetStringUTFChars(NextId, 0);
+			NewProductInfo.Identifier = FString(UTF8_TO_TCHAR(charsId));
+			jenv->ReleaseStringUTFChars(NextId, charsId);
+			jenv->DeleteLocalRef(NextId);
+
+			jstring NextTitle = (jstring)jenv->GetObjectArrayElement(titles, Idx);
+			const char* charsTitle = jenv->GetStringUTFChars(NextTitle, 0);
+			NewProductInfo.DisplayName = FString(UTF8_TO_TCHAR(charsTitle));
+			jenv->ReleaseStringUTFChars(NextTitle, charsTitle);
+			jenv->DeleteLocalRef(NextTitle);
+
+			jstring NextDesc = (jstring)jenv->GetObjectArrayElement(descriptions, Idx);
+			const char* charsDesc = jenv->GetStringUTFChars(NextDesc, 0);
+			NewProductInfo.DisplayDescription = FString(UTF8_TO_TCHAR(charsDesc));
+			jenv->ReleaseStringUTFChars(NextDesc, charsDesc);
+			jenv->DeleteLocalRef(NextDesc);
+
+			jstring NextPrice = (jstring)jenv->GetObjectArrayElement(prices, Idx);
+			const char* charsPrice = jenv->GetStringUTFChars(NextPrice, 0);
+			NewProductInfo.DisplayPrice = FString(UTF8_TO_TCHAR(charsPrice));
+			jenv->ReleaseStringUTFChars(NextPrice, charsPrice);
+			jenv->DeleteLocalRef(NextPrice);
+
+			NewProductInfo.RawPrice = PricesRaw[Idx];
+
+			jstring NextCurrencyCode = (jstring)jenv->GetObjectArrayElement(currencyCodes, Idx);
+			const char* charsCurrencyCodes = jenv->GetStringUTFChars(NextCurrencyCode, 0);
+			NewProductInfo.CurrencyCode = FString(UTF8_TO_TCHAR(charsCurrencyCodes));
+			jenv->ReleaseStringUTFChars(NextCurrencyCode, charsCurrencyCodes);
+			jenv->DeleteLocalRef(NextCurrencyCode);
+
+			ProvidedProductInformation.Add(NewProductInfo);
+
+			FPlatformMisc::LowLevelOutputDebugStringf(TEXT("\nProduct Identifier: %s, Name: %s, Description: %s, Price: %s, Price Raw: %f, Currency Code: %s\n"),
+				*NewProductInfo.Identifier,
+				*NewProductInfo.DisplayName,
+				*NewProductInfo.DisplayDescription,
+				*NewProductInfo.DisplayPrice,
+				NewProductInfo.RawPrice,
+				*NewProductInfo.CurrencyCode);
+		}
+	}
+
+
+	DECLARE_CYCLE_STAT(TEXT("FSimpleDelegateGraphTask.ProcessQueryIapResult"), STAT_FSimpleDelegateGraphTask_ProcessQueryIapResult, STATGROUP_TaskGraphTasks);
+
+	FSimpleDelegateGraphTask::CreateAndDispatchWhenReady(
+		FSimpleDelegateGraphTask::FDelegate::CreateLambda([=](){
+			if (IOnlineSubsystem* const OnlineSub = IOnlineSubsystem::Get())
+			{
+				// call store implementation to process query results.
+				if (FOnlineStoreGooglePlay* StoreInterface = (FOnlineStoreGooglePlay*)OnlineSub->GetStoreInterface().Get())
+				{
+					StoreInterface->ProcessQueryAvailablePurchasesResults(bSuccess, ProvidedProductInformation);
+				}
+			}
+			FPlatformMisc::LowLevelOutputDebugStringf(TEXT("In-App Purchase query was completed  %s\n"), bSuccess ? TEXT("successfully") : TEXT("unsuccessfully"));
+		}),
+		GET_STATID(STAT_FSimpleDelegateGraphTask_ProcessQueryIapResult), 
+		nullptr, 
+		ENamedThreads::GameThread
+	);
+
+}
+
+
+void FOnlineStoreGooglePlay::ProcessQueryAvailablePurchasesResults(bool bInSuccessful, const TArray<FInAppPurchaseProductInfo>& AvailablePurchases)
+{
+	UE_LOG(LogOnline, Display, TEXT("FOnlineStoreGooglePlay::ProcessQueryAvailablePurchasesResults"));
+
+	if (ReadObject.IsValid())
+	{
+		ReadObject->ReadState = bInSuccessful ? EOnlineAsyncTaskState::Done : EOnlineAsyncTaskState::Failed;
+		ReadObject->ProvidedProductInformation.Insert(AvailablePurchases, 0);
+	}
+
+	CurrentQueryTask->ProcessQueryAvailablePurchasesResults(bInSuccessful);
+}
+
+
+bool FOnlineStoreGooglePlay::BeginPurchase(const FInAppPurchaseProductRequest& ProductRequest, FOnlineInAppPurchaseTransactionRef& InPurchaseStateObject)
+{
+	UE_LOG(LogOnline, Display, TEXT( "FOnlineStoreGooglePlay::BeginPurchase" ));
+	
+	bool bCreatedNewTransaction = false;
+	
+	if (IsAllowedToMakePurchases())
+	{
+		CachedPurchaseStateObject = InPurchaseStateObject;
+
+		extern bool AndroidThunkCpp_Iap_BeginPurchase(const FString&, const bool);
+		bCreatedNewTransaction = AndroidThunkCpp_Iap_BeginPurchase(ProductRequest.ProductIdentifier, ProductRequest.bIsConsumable);
+		UE_LOG(LogOnline, Display, TEXT("Created Transaction? - %s"), 
+			bCreatedNewTransaction ? TEXT("Created a transaction.") : TEXT("Failed to create a transaction."));
+
+		if (!bCreatedNewTransaction)
+		{
+			UE_LOG(LogOnline, Display, TEXT("FOnlineStoreGooglePlay::BeginPurchase - Could not create a new transaction."));
+			CachedPurchaseStateObject->ReadState = EOnlineAsyncTaskState::Failed;
+			TriggerOnInAppPurchaseCompleteDelegates(EInAppPurchaseState::Invalid);
+		}
+		else
+		{
+			CachedPurchaseStateObject->ReadState = EOnlineAsyncTaskState::InProgress;
+		}
+	}
+	else
+	{
+		UE_LOG(LogOnline, Display, TEXT("This device is not able to make purchases."));
+
+		InPurchaseStateObject->ReadState = EOnlineAsyncTaskState::Failed;
+		TriggerOnInAppPurchaseCompleteDelegates(EInAppPurchaseState::NotAllowed);
+	}
+
+
+	return bCreatedNewTransaction;
+}
+
+
+extern "C" void Java_com_epicgames_ue4_GooglePlayStoreHelper_nativePurchaseComplete(JNIEnv* jenv, jobject thiz, jboolean bSuccess, jstring productId, jstring receiptData, jstring signature)
+{
+	FString ProductId, ReceiptData, Signature;
+	if (bSuccess)
+	{
+		const char* charsId = jenv->GetStringUTFChars(productId, 0);
+		ProductId = FString(UTF8_TO_TCHAR(charsId));
+		jenv->ReleaseStringUTFChars(productId, charsId);
+
+		const char* charsReceipt = jenv->GetStringUTFChars(receiptData, 0);
+		ReceiptData = FString(UTF8_TO_TCHAR(charsReceipt));
+		jenv->ReleaseStringUTFChars(receiptData, charsReceipt);
+
+		const char* charsSignature = jenv->GetStringUTFChars(signature, 0);
+		Signature = FString(UTF8_TO_TCHAR(charsSignature));
+		jenv->ReleaseStringUTFChars(signature, charsSignature);
+	}
+	FPlatformMisc::LowLevelOutputDebugStringf(TEXT("1... ProductId: %s, ReceiptData: %s, Signature: %s\n"), *ProductId, *ReceiptData, *Signature );
+
+	DECLARE_CYCLE_STAT(TEXT("FSimpleDelegateGraphTask.ProcessIapResult"), STAT_FSimpleDelegateGraphTask_ProcessIapResult, STATGROUP_TaskGraphTasks);
+
+	FSimpleDelegateGraphTask::CreateAndDispatchWhenReady(
+		FSimpleDelegateGraphTask::FDelegate::CreateLambda([=](){
+			FPlatformMisc::LowLevelOutputDebugStringf(TEXT("In-App Purchase was completed  %s\n"), bSuccess ? TEXT("successfully") : TEXT("unsuccessfully"));
+			if (IOnlineSubsystem* const OnlineSub = IOnlineSubsystem::Get())
+			{
+				FPlatformMisc::LowLevelOutputDebugStringf(TEXT("2... ProductId: %s, ReceiptData: %s, Signature: %s\n"), *ProductId, *ReceiptData, *Signature);
+				// call store implementation to process query results.
+				if (FOnlineStoreGooglePlay* StoreInterface = (FOnlineStoreGooglePlay*)OnlineSub->GetStoreInterface().Get())
+				{
+					StoreInterface->ProcessPurchaseResult(bSuccess, ProductId, ReceiptData, Signature);
+				}
+			}
+		}),
+		GET_STATID(STAT_FSimpleDelegateGraphTask_ProcessIapResult), 
+		nullptr, 
+		ENamedThreads::GameThread
+	);
+
+}
+
+
+void FOnlineStoreGooglePlay::ProcessPurchaseResult(bool bInSuccessful, const FString& ProductId, const FString& InReceiptData, const FString& Signature)
+{
+	UE_LOG(LogOnline, Display, TEXT("FOnlineStoreGooglePlay::ProcessPurchaseResult"));
+	UE_LOG(LogOnline, Display, TEXT("3... ProductId: %s, ReceiptData: %s, Signature: %s\n"), *ProductId, *InReceiptData, *Signature);
+
+
+	if (CachedPurchaseStateObject.IsValid())
+	{
+		FInAppPurchaseProductInfo& ProductInfo = CachedPurchaseStateObject->ProvidedProductInformation;
+		ProductInfo.Identifier = ProductId;
+		ProductInfo.DisplayName = TEXT("n/a");
+		ProductInfo.DisplayDescription = TEXT("n/a");
+		ProductInfo.DisplayPrice = TEXT("n/a");
+		ProductInfo.ReceiptData = InReceiptData;
+		ProductInfo.TransactionIdentifier = Signature;
+
+		CachedPurchaseStateObject->ReadState = EOnlineAsyncTaskState::Done;
+	}
+
+	TriggerOnInAppPurchaseCompleteDelegates(bInSuccessful ? EInAppPurchaseState::Success : EInAppPurchaseState::Failed);
+}
+
+
+bool FOnlineStoreGooglePlay::RestorePurchases(const TArray<FInAppPurchaseProductRequest>& ConsumableProductFlags, FOnlineInAppPurchaseRestoreReadRef& InReadObject)
+{
+	bool bSentAQueryRequest = false;
+	CachedPurchaseRestoreObject = InReadObject;
+
+	if (IsAllowedToMakePurchases())
+	{
+		TArray<FString> ProductIds;
+		TArray<bool> IsConsumableFlags;
+
+		for (int i = 0; i < ConsumableProductFlags.Num(); i++)
+		{
+			ProductIds.Add(ConsumableProductFlags[i].ProductIdentifier);
+			IsConsumableFlags.Add(ConsumableProductFlags[i].bIsConsumable);
+		}
+
+		// Send JNI request
+		extern bool AndroidThunkCpp_Iap_RestorePurchases(const TArray<FString>&, const TArray<bool>&);
+		bSentAQueryRequest = AndroidThunkCpp_Iap_RestorePurchases(ProductIds, IsConsumableFlags);
+	}
+	else
+	{
+		UE_LOG(LogOnline, Display, TEXT("This device is not able to make purchases."));
+		TriggerOnInAppPurchaseRestoreCompleteDelegates(EInAppPurchaseState::Failed);
+	}
+
+	return bSentAQueryRequest;
+}
+
+extern "C" void Java_com_epicgames_ue4_GooglePlayStoreHelper_nativeRestorePurchasesComplete(JNIEnv* jenv, jobject thiz, jboolean bSuccess, jobjectArray ProductIDs, jobjectArray ReceiptsData, jobjectArray Signatures)
+{
+	TArray<FInAppPurchaseRestoreInfo> RestoredPurchaseInfo;
+
+	if (jenv && bSuccess)
+	{
+		jsize NumProducts = jenv->GetArrayLength(ProductIDs);
+		jsize NumReceipts = jenv->GetArrayLength(ReceiptsData);
+		jsize NumSignatures = jenv->GetArrayLength(Signatures);
+
+		ensure((NumProducts == NumReceipts) && (NumProducts == NumSignatures));
+
+		for (jsize Idx = 0; Idx < NumProducts; Idx++)
+		{
+			// Build the restore product information strings.
+			FInAppPurchaseRestoreInfo RestoreInfo;
+
+			jstring NextId = (jstring)jenv->GetObjectArrayElement(ProductIDs, Idx);
+			const char* charsId = jenv->GetStringUTFChars(NextId, 0);
+			RestoreInfo.Identifier = FString(UTF8_TO_TCHAR(charsId));
+			jenv->ReleaseStringUTFChars(NextId, charsId);
+			jenv->DeleteLocalRef(NextId);
+
+			jstring NextReceipt = (jstring)jenv->GetObjectArrayElement(ReceiptsData, Idx);
+			const char* charsReceipt = jenv->GetStringUTFChars(NextReceipt, 0);
+			RestoreInfo.ReceiptData = FString(UTF8_TO_TCHAR(charsReceipt));
+			jenv->ReleaseStringUTFChars(NextReceipt, charsReceipt);
+			jenv->DeleteLocalRef(NextReceipt);
+
+			jstring NextSignature = (jstring)jenv->GetObjectArrayElement(Signatures, Idx);
+			const char* charsSignature = jenv->GetStringUTFChars(NextSignature, 0);
+			RestoreInfo.TransactionIdentifier = FString(UTF8_TO_TCHAR(charsSignature));
+			jenv->ReleaseStringUTFChars(NextSignature, charsSignature);
+			jenv->DeleteLocalRef(NextSignature);
+
+			RestoredPurchaseInfo.Add(RestoreInfo);
+
+			FPlatformMisc::LowLevelOutputDebugStringf(TEXT("\nRestored Product Identifier: %s\n"), *RestoreInfo.Identifier);
+		}
+	}
+
+	DECLARE_CYCLE_STAT(TEXT("FSimpleDelegateGraphTask.RestorePurchases"), STAT_FSimpleDelegateGraphTask_RestorePurchases, STATGROUP_TaskGraphTasks);
+
+	FSimpleDelegateGraphTask::CreateAndDispatchWhenReady(
+		FSimpleDelegateGraphTask::FDelegate::CreateLambda([=]()
+		{
+			FPlatformMisc::LowLevelOutputDebugStringf(TEXT("Restoring In-App Purchases was completed  %s\n"), bSuccess ? TEXT("successfully") : TEXT("unsuccessfully"));
+			if (IOnlineSubsystem* const OnlineSub = IOnlineSubsystem::Get())
+			{
+				FPlatformMisc::LowLevelOutputDebugStringf(TEXT("Sending result back to OnlineSubsystem.\n"));
+				// call store implementation to process query results.
+				if (FOnlineStoreGooglePlay* StoreInterface = (FOnlineStoreGooglePlay*)OnlineSub->GetStoreInterface().Get())
+				{
+					if (StoreInterface->CachedPurchaseRestoreObject.IsValid())
+					{
+						StoreInterface->CachedPurchaseRestoreObject->ProvidedRestoreInformation = RestoredPurchaseInfo;
+						StoreInterface->CachedPurchaseRestoreObject->ReadState = bSuccess ? EOnlineAsyncTaskState::Done : EOnlineAsyncTaskState::Failed;
+					}
+					StoreInterface->TriggerOnInAppPurchaseRestoreCompleteDelegates(bSuccess ? EInAppPurchaseState::Restored : EInAppPurchaseState::Failed);
+				}
+			}
+		}),
+		GET_STATID(STAT_FSimpleDelegateGraphTask_RestorePurchases),
+		nullptr,
+		ENamedThreads::GameThread
+	);
+}
\ No newline at end of file
diff --git a/Engine/Plugins/Online/Android/OnlineSubsystemGooglePlay/Source/Public/OnlineSubsystemGooglePlay.h b/Engine/Plugins/Online/Android/OnlineSubsystemGooglePlay/Source/Public/OnlineSubsystemGooglePlay.h
new file mode 100644
--- /dev/null
+++ b/Engine/Plugins/Online/Android/OnlineSubsystemGooglePlay/Source/Public/OnlineSubsystemGooglePlay.h
@@ -0,0 +1,188 @@
+// Copyright 1998-2016 Epic Games, Inc. All Rights Reserved.
+
+#pragma once
+
+#include "OnlineSubsystem.h"
+#include "OnlineSubsystemImpl.h"
+#include "OnlineIdentityInterfaceGooglePlay.h"
+#include "OnlineAchievementsInterfaceGooglePlay.h"
+#include "OnlineLeaderboardInterfaceGooglePlay.h"
+#include "OnlineExternalUIInterfaceGooglePlay.h"
+#include "OnlineAsyncTaskManagerGooglePlay.h"
+#include "UniquePtr.h"
+
+#include "gpg/game_services.h"
+#include "gpg/android_platform_configuration.h"
+
+class FOnlineAsyncTaskGooglePlayLogin;
+class FOnlineAsyncTaskGooglePlayShowLoginUI;
+class FOnlineAsyncTaskGooglePlayLogout;
+class FRunnableThread;
+
+/**
+ * OnlineSubsystemGooglePlay - Implementation of the online subsystem for Google Play services
+ */
+class ONLINESUBSYSTEMGOOGLEPLAY_API FOnlineSubsystemGooglePlay : 
+	public FOnlineSubsystemImpl
+{
+public:
+	FOnlineSubsystemGooglePlay();
+	virtual ~FOnlineSubsystemGooglePlay() {}
+
+	//~ Begin IOnlineSubsystem Interface
+	virtual IOnlineSessionPtr GetSessionInterface() const override;
+	virtual IOnlineFriendsPtr GetFriendsInterface() const override;
+	virtual IOnlinePartyPtr GetPartyInterface() const override;
+	virtual IOnlineGroupsPtr GetGroupsInterface() const override;
+	virtual IOnlineSharedCloudPtr GetSharedCloudInterface() const override;
+	virtual IOnlineUserCloudPtr GetUserCloudInterface() const override;
+	virtual IOnlineLeaderboardsPtr GetLeaderboardsInterface() const override;
+	virtual IOnlineVoicePtr GetVoiceInterface() const  override;
+	virtual IOnlineExternalUIPtr GetExternalUIInterface() const override;
+	virtual IOnlineTimePtr GetTimeInterface() const override;
+	virtual IOnlineIdentityPtr GetIdentityInterface() const override;
+	virtual IOnlineTitleFilePtr GetTitleFileInterface() const override;
+	virtual IOnlineEntitlementsPtr GetEntitlementsInterface() const override;
+	virtual IOnlineStorePtr GetStoreInterface() const override;
+	virtual IOnlineStoreV2Ptr GetStoreV2Interface() const override { return nullptr; }
+	virtual IOnlinePurchasePtr GetPurchaseInterface() const override { return nullptr; }
+	virtual IOnlineEventsPtr GetEventsInterface() const override { return nullptr; }
+	virtual IOnlineMessagePtr GetMessageInterface() const override { return nullptr; }
+	virtual IOnlineSharingPtr GetSharingInterface() const override { return nullptr; }
+	virtual IOnlineUserPtr GetUserInterface() const override { return nullptr; }
+	virtual IOnlineAchievementsPtr GetAchievementsInterface() const override;
+	virtual IOnlinePresencePtr GetPresenceInterface() const override { return nullptr; }
+	virtual IOnlineChatPtr GetChatInterface() const override { return nullptr; }
+	virtual IOnlineTurnBasedPtr GetTurnBasedInterface() const override { return nullptr; }
+
+	virtual class UObject* GetNamedInterface(FName InterfaceName) override { return nullptr; }
+	virtual void SetNamedInterface(FName InterfaceName, class UObject* NewInterface) override {}
+	virtual bool IsDedicated() const override { return false; }
+	virtual bool IsServer() const override { return true; }
+	virtual void SetForceDedicated(bool bForce) override {}
+	virtual bool IsLocalPlayer(const FUniqueNetId& UniqueId) const override { return true; }
+
+	virtual bool Init() override;
+	virtual bool Shutdown() override;
+	virtual FString GetAppId() const override;
+	virtual bool Exec(class UWorld* InWorld, const TCHAR* Cmd, FOutputDevice& Ar) override;
+	//~ End IOnlineSubsystem Interface
+
+	virtual bool Tick(float DeltaTime) override;
+	
+PACKAGE_SCOPE:
+
+	/**
+	 * Is Online Subsystem Android available for use
+	 * @return true if Android Online Subsystem functionality is available, false otherwise
+	 */
+	bool IsEnabled();
+
+	/** Return the async task manager owned by this subsystem */
+	class FOnlineAsyncTaskManagerGooglePlay* GetAsyncTaskManager() { return OnlineAsyncTaskThreadRunnable.Get(); }
+
+	/**
+	 * Add an async task onto the task queue for processing
+	 * @param AsyncTask - new heap allocated task to process on the async task thread
+	 */
+	void QueueAsyncTask(class FOnlineAsyncTask* AsyncTask);
+
+	/** Returns a pointer to the Google API entry point */
+	gpg::GameServices* GetGameServices() const { return GameServicesPtr.get(); }
+
+	/** Utility function, useful for Google APIs that take a std::string but we only have an FString */
+	static std::string ConvertFStringToStdString(const FString& InString);
+
+	/** Returns the Google Play-specific version of Identity, useful to avoid unnecessary casting */
+	FOnlineIdentityGooglePlayPtr GetIdentityGooglePlay() const { return IdentityInterface; }
+
+	/** Returns the Google Play-specific version of Achievements, useful to avoid unnecessary casting */
+	FOnlineAchievementsGooglePlayPtr GetAchievementsGooglePlay() const { return AchievementsInterface; }
+
+	/**
+	* Is IAP available for use
+	* @return true if IAP should be available, false otherwise
+	*/
+	bool IsInAppPurchasingEnabled();
+
+	/** Returns true if there are any async login tasks currently being tracked. */
+	bool AreAnyAsyncLoginTasksRunning() const { return CurrentLoginTask != nullptr || CurrentShowLoginUITask != nullptr; }
+
+	/** Start a ShowLoginUI async task. Creates the GameServices object first if necessary. */
+	void StartShowLoginUITask(int PlayerId, const FOnLoginUIClosedDelegate& Delegate = FOnLoginUIClosedDelegate());
+
+	/** Start a logout task if one isn't already in progress. */
+	void StartLogoutTask(int32 LocalUserNum);
+
+private:
+
+	/** Google callback when auth is complete */
+	void OnAuthActionFinished(gpg::AuthOperation Op, gpg::AuthStatus Status);
+
+	/** Android callback when an activity is finished */
+	void OnActivityResult(JNIEnv *env, jobject thiz, jobject activity, jint requestCode, jint resultCode, jobject data);
+
+	/** Start a ShowLoginUI async task. */
+	void StartShowLoginUITask_Internal(int PlayerId, const FOnLoginUIClosedDelegate& Delegate = FOnLoginUIClosedDelegate());
+
+	/** Online async task runnable */
+	TUniquePtr<class FOnlineAsyncTaskManagerGooglePlay> OnlineAsyncTaskThreadRunnable;
+
+	/** Online async task thread */
+	TUniquePtr<FRunnableThread> OnlineAsyncTaskThread;
+
+	/** Interface to the online identity system */
+	FOnlineIdentityGooglePlayPtr IdentityInterface;
+
+	FOnlineStoreGooglePlayPtr StoreInterface;
+
+	/** Interface to the online leaderboards */
+	FOnlineLeaderboardsGooglePlayPtr LeaderboardsInterface;
+
+	/** Interface to the online achievements */
+	FOnlineAchievementsGooglePlayPtr AchievementsInterface;
+
+	/** Interface to the external UI services */
+	FOnlineExternalUIGooglePlayPtr ExternalUIInterface;
+
+	/** Pointer to the main entry point for the Google API */
+	std::unique_ptr<gpg::GameServices> GameServicesPtr;
+
+	// Since the OnAuthActionFinished callback is "global" (within the scope of one GameServices object),
+	// we track the async tasks that depend on this callback to notify us that they're finished.
+	// Only one of these pointers should be non-null at any given time.
+
+	/**
+	 * Track the current login task (if any) so callbacks can notify it.
+	 * Still owned by the AsyncTaskManager, do not delete via this pointer!
+	 **/
+	FOnlineAsyncTaskGooglePlayLogin* CurrentLoginTask;
+
+	/**
+	 * Track the current ShowLoginUI task (if any)
+	 * Still owned by the AsyncTaskManager, do not delete via this pointer!
+	 **/
+	FOnlineAsyncTaskGooglePlayShowLoginUI* CurrentShowLoginUITask;
+
+	/**
+	 * Track the current Logout task
+	 * Still owned by the AsyncTaskManager, do not delete via this pointer!
+	 */
+	FOnlineAsyncTaskGooglePlayLogout* CurrentLogoutTask;
+
+	/** Handle of registered delegate for onActivityResult */
+	FDelegateHandle OnActivityResultDelegateHandle;
+
+	/** This task needs to be able to set the GameServicesPtr and clear CurrentLoginTask*/
+	friend class FOnlineAsyncTaskGooglePlayLogin;
+
+	/** This task needs to be able to clear CurrentShowLoginUITask */
+	friend class FOnlineAsyncTaskGooglePlayShowLoginUI;
+
+	/** This task needs to be able to clear CurrentLogoutTask */
+	friend class FOnlineAsyncTaskGooglePlayLogout;
+
+	gpg::AndroidPlatformConfiguration PlatformConfiguration;
+};
+
+typedef TSharedPtr<FOnlineSubsystemGooglePlay, ESPMode::ThreadSafe> FOnlineSubsystemGooglePlayPtr;
diff --git a/Engine/Plugins/Online/IOS/OnlineSubsystemIOS/Source/OnlineSubsystemIOS.Build.cs b/Engine/Plugins/Online/IOS/OnlineSubsystemIOS/Source/OnlineSubsystemIOS.Build.cs
new file mode 100644
--- /dev/null
+++ b/Engine/Plugins/Online/IOS/OnlineSubsystemIOS/Source/OnlineSubsystemIOS.Build.cs
@@ -0,0 +1,34 @@
+// Copyright 1998-2016 Epic Games, Inc. All Rights Reserved.
+
+using UnrealBuildTool;
+
+public class OnlineSubsystemIOS : ModuleRules
+{
+    public OnlineSubsystemIOS(TargetInfo Target)
+	{
+		PrivateIncludePaths.AddRange( 
+            new string[] {
+                "Private",             
+                }
+                );
+
+        Definitions.Add("ONLINESUBSYSTEMIOS_PACKAGE=1");
+
+		PrivateDependencyModuleNames.AddRange(
+            new string[] { 
+                "Core", 
+                "CoreUObject", 
+                "Engine", 
+                "Sockets",
+				"OnlineSubsystem", 
+                "Http",
+            }
+            );
+
+		PublicWeakFrameworks.Add("Cloudkit");
+		if (Target.Platform == UnrealTargetPlatform.IOS)
+		{
+			PublicWeakFrameworks.Add("MultipeerConnectivity");
+		}
+	}
+}
diff --git a/Engine/Plugins/Online/IOS/OnlineSubsystemIOS/Source/Private/OnlineFriendsInterfaceIOS.cpp b/Engine/Plugins/Online/IOS/OnlineSubsystemIOS/Source/Private/OnlineFriendsInterfaceIOS.cpp
new file mode 100644
--- /dev/null
+++ b/Engine/Plugins/Online/IOS/OnlineSubsystemIOS/Source/Private/OnlineFriendsInterfaceIOS.cpp
@@ -0,0 +1,310 @@
+// Copyright 1998-2016 Epic Games, Inc. All Rights Reserved.
+
+
+#include "OnlineSubsystemIOSPrivatePCH.h"
+
+// FOnlineFriendIOS
+
+TSharedRef<const FUniqueNetId> FOnlineFriendIOS::GetUserId() const
+{
+	return UserId;
+}
+
+FString FOnlineFriendIOS::GetRealName() const
+{
+	FString Result;
+	GetAccountData(TEXT("nickname"), Result);
+	return Result;
+}
+
+FString FOnlineFriendIOS::GetDisplayName(const FString& Platform) const
+{
+	FString Result;
+	GetAccountData(TEXT("nickname"), Result);
+	return Result;
+}
+
+bool FOnlineFriendIOS::GetUserAttribute(const FString& AttrName, FString& OutAttrValue) const
+{
+	return GetAccountData(AttrName, OutAttrValue);
+}
+
+EInviteStatus::Type FOnlineFriendIOS::GetInviteStatus() const
+{
+	return EInviteStatus::Accepted;
+}
+
+const FOnlineUserPresence& FOnlineFriendIOS::GetPresence() const
+{
+	return Presence;
+}
+
+// FOnlineFriendsIOS
+
+FOnlineFriendsIOS::FOnlineFriendsIOS(FOnlineSubsystemIOS* InSubsystem)
+{
+	UE_LOG(LogOnline, Verbose, TEXT("FOnlineFriendsIOS::FOnlineFriendsIOS()"));
+
+	IdentityInterface = (FOnlineIdentityIOS*)InSubsystem->GetIdentityInterface().Get();
+}
+
+bool FOnlineFriendsIOS::ReadFriendsList(int32 LocalUserNum, const FString& ListName, const FOnReadFriendsListComplete& Delegate /*= FOnReadFriendsListComplete()*/)
+{
+	UE_LOG(LogOnline, Verbose, TEXT("FOnlineFriendsIOS::ReadFriendsList()"));
+	bool bSuccessfullyBeganReadFriends = false;
+
+	if( IdentityInterface->GetLocalGameCenterUser() != NULL && IdentityInterface->GetLocalGameCenterUser().isAuthenticated )
+	{
+		bSuccessfullyBeganReadFriends = true;
+		dispatch_async(dispatch_get_main_queue(), ^
+		{
+		// Get the friends list for the local player from the server
+#ifdef __IPHONE_8_0
+            if ([[GKLocalPlayer localPlayer] respondsToSelector:@selector(loadFriendPlayersWithCompletionHandler)] == YES)
+            {
+                [[GKLocalPlayer localPlayer] loadFriendPlayersWithCompletionHandler:
+                 ^(NSArray* Friends, NSError* Error)
+                 {
+                    if( Error )
+                    {
+                        FString ErrorStr(FString::Printf(TEXT("FOnlineFriendsIOS::ReadFriendsList() - Failed to read friends list with error: [%i]"), [Error code]));
+                        UE_LOG(LogOnline, Verbose, TEXT("%s"), *ErrorStr);
+                 
+                        // Report back to the game thread whether this succeeded.
+                        [FIOSAsyncTask CreateTaskWithBlock : ^ bool(void)
+                         {
+                            Delegate.ExecuteIfBound(0, false, ListName, ErrorStr);
+                            return true;
+                         }];
+                    }
+                    else
+                    {
+                        [GKPlayer loadPlayersForIdentifiers:Friends withCompletionHandler:
+                         ^(NSArray* Players, NSError* Error2)
+                         {
+                            FString ErrorStr;
+                            bool bWasSuccessful = false;
+                            if( Error2 )
+                            {
+                                ErrorStr = FString::Printf(TEXT("FOnlineFriendsIOS::ReadFriendsList() - Failed to loadPlayersForIdentifiers with error: [%i]"), [Error2 code]);
+                            }
+                            else
+                            {
+                                // Clear our previosly cached friends before we repopulate the cache.
+                                CachedFriends.Empty();
+                                for( int32 FriendIdx = 0; FriendIdx < [Players count]; FriendIdx++ )
+                                {
+                                    GKPlayer* Friend = Players[ FriendIdx ];
+                  
+                                    // Add new friend entry to list
+                                    TSharedRef<FOnlineFriendIOS> FriendEntry(
+                                            new FOnlineFriendIOS(*FString(Friend.playerID))
+                                    );
+                                    FriendEntry->AccountData.Add(
+                                            TEXT("nickname"), *FString(Friend.alias)
+                                    );
+                                    CachedFriends.Add(FriendEntry);
+                  
+                                    UE_LOG(LogOnline, Verbose, TEXT("GCFriend - Id:%s Alias:%s"), *FString(Friend.playerID),
+                                        *FriendEntry->GetDisplayName() );
+                                }
+                  
+                                bWasSuccessful = true;
+                            }
+                  
+                            // Report back to the game thread whether this succeeded.
+                            [FIOSAsyncTask CreateTaskWithBlock : ^ bool(void)
+                             {
+                                Delegate.ExecuteIfBound(0, bWasSuccessful, ListName, ErrorStr);
+                                return true;
+                             }];
+                         }];
+                    }
+                 }];
+            }
+            else
+#endif
+            {
+#if __IPHONE_OS_VERSION_MIN_REQUIRED < __IPHONE_8_0
+                [[GKLocalPlayer localPlayer] loadFriendsWithCompletionHandler:
+                 ^(NSArray* Friends, NSError* Error)
+                 {
+                    if( Error )
+                    {
+                        FString ErrorStr(FString::Printf(TEXT("FOnlineFriendsIOS::ReadFriendsList() - Failed to read friends list with error: [%i]"), [Error code]));
+                        UE_LOG(LogOnline, Verbose, TEXT("%s"), *ErrorStr);
+
+						// Report back to the game thread whether this succeeded.
+						[FIOSAsyncTask CreateTaskWithBlock : ^ bool(void)
+						{
+							Delegate.ExecuteIfBound(0, false, ListName, ErrorStr);
+							return true;
+						}];
+                    }
+                    else
+                    {
+                        [GKPlayer loadPlayersForIdentifiers:Friends withCompletionHandler:
+                         ^(NSArray* Players, NSError* Error2)
+                         {
+								FString ErrorStr;
+								bool bWasSuccessful = false;
+							if( Error2 )
+							{
+									ErrorStr = FString::Printf(TEXT("FOnlineFriendsIOS::ReadFriendsList() - Failed to loadPlayersForIdentifiers with error: [%i]"), [Error2 code]);
+							}
+							else
+							{
+								// Clear our previosly cached friends before we repopulate the cache.
+								CachedFriends.Empty();
+								for( int32 FriendIdx = 0; FriendIdx < [Players count]; FriendIdx++ )
+								{
+									GKPlayer* Friend = Players[ FriendIdx ];
+
+									// Add new friend entry to list
+									TSharedRef<FOnlineFriendIOS> FriendEntry(
+										new FOnlineFriendIOS(*FString(Friend.playerID))
+										);
+									FriendEntry->AccountData.Add(
+										TEXT("nickname"), *FString(Friend.alias)
+										);
+									CachedFriends.Add(FriendEntry);
+
+									UE_LOG(LogOnline, Verbose, TEXT("GCFriend - Id:%s Alias:%s"), *FString(Friend.playerID), 
+										*FriendEntry->GetDisplayName() );
+								}
+
+                                bWasSuccessful = true;
+                            }
+								
+							// Report back to the game thread whether this succeeded.
+							[FIOSAsyncTask CreateTaskWithBlock : ^ bool(void)
+							{
+								Delegate.ExecuteIfBound(0, bWasSuccessful, ListName, ErrorStr);
+								return true;
+							}];
+						}];
+                    }
+                }];
+#endif
+			}
+		});
+	}
+	else
+	{
+		// no gamecenter login means we cannot read the friends.
+		Delegate.ExecuteIfBound( 0 , false, ListName, TEXT("not logged in") );
+	}
+
+	return bSuccessfullyBeganReadFriends;
+}
+
+bool FOnlineFriendsIOS::DeleteFriendsList(int32 LocalUserNum, const FString& ListName, const FOnDeleteFriendsListComplete& Delegate /*= FOnDeleteFriendsListComplete()*/)
+{
+	Delegate.ExecuteIfBound(LocalUserNum, false, ListName, FString(TEXT("DeleteFriendsList() is not supported")));
+	return false;
+}
+
+bool FOnlineFriendsIOS::SendInvite(int32 LocalUserNum, const FUniqueNetId& FriendId, const FString& ListName, const FOnSendInviteComplete& Delegate /*= FOnSendInviteComplete()*/)
+{
+	Delegate.ExecuteIfBound(LocalUserNum, false, FriendId, ListName, FString(TEXT("SendInvite() is not supported")));
+	return false;
+}
+
+bool FOnlineFriendsIOS::AcceptInvite(int32 LocalUserNum, const FUniqueNetId& FriendId, const FString& ListName, const FOnAcceptInviteComplete& Delegate /*= FOnAcceptInviteComplete()*/)
+{
+	Delegate.ExecuteIfBound(LocalUserNum, false, FriendId, ListName, FString(TEXT("AcceptInvite() is not supported")));
+	return false;
+}
+
+bool FOnlineFriendsIOS::RejectInvite(int32 LocalUserNum, const FUniqueNetId& FriendId, const FString& ListName)
+{
+	TriggerOnRejectInviteCompleteDelegates(LocalUserNum, false, FriendId, ListName, FString(TEXT("RejectInvite() is not supported")));
+	return false;
+}
+
+bool FOnlineFriendsIOS::DeleteFriend(int32 LocalUserNum, const FUniqueNetId& FriendId, const FString& ListName)
+{
+	TriggerOnDeleteFriendCompleteDelegates(LocalUserNum, false, FriendId, ListName, FString(TEXT("DeleteFriend() is not supported")));
+	return false;
+}
+
+bool FOnlineFriendsIOS::GetFriendsList(int32 LocalUserNum, const FString& ListName, TArray< TSharedRef<FOnlineFriend> >& OutFriends)
+{
+	UE_LOG(LogOnline, Verbose, TEXT("FOnlineFriendsIOS::GetFriendsList()"));
+
+	for (int32 Idx=0; Idx < CachedFriends.Num(); Idx++)
+	{
+		OutFriends.Add(CachedFriends[Idx]);
+	}
+
+	return true;
+}
+
+TSharedPtr<FOnlineFriend> FOnlineFriendsIOS::GetFriend(int32 LocalUserNum, const FUniqueNetId& FriendId, const FString& ListName)
+{
+	TSharedPtr<FOnlineFriend> Result;
+
+	UE_LOG(LogOnline, Verbose, TEXT("FOnlineFriendsIOS::GetFriend()"));
+
+	for (int32 Idx=0; Idx < CachedFriends.Num(); Idx++)
+	{
+		if (*(CachedFriends[Idx]->GetUserId()) == FriendId)
+		{
+			Result = CachedFriends[Idx];
+			break;
+		}
+	}
+
+	return Result;
+}
+
+bool FOnlineFriendsIOS::IsFriend(int32 LocalUserNum, const FUniqueNetId& FriendId, const FString& ListName)
+{
+	UE_LOG(LogOnline, Verbose, TEXT("FOnlineFriendsIOS::IsFriend()"));
+
+	TSharedPtr<FOnlineFriend> Friend = GetFriend(LocalUserNum,FriendId,ListName);
+	if (Friend.IsValid() &&
+		Friend->GetInviteStatus() == EInviteStatus::Accepted)
+	{
+		return true;
+	}
+	return false;
+}
+
+bool FOnlineFriendsIOS::QueryRecentPlayers(const FUniqueNetId& UserId, const FString& Namespace)
+{
+	UE_LOG(LogOnline, Verbose, TEXT("FOnlineFriendsIOS::QueryRecentPlayers()"));
+
+	TriggerOnQueryRecentPlayersCompleteDelegates(UserId, Namespace, false, TEXT("not implemented"));
+
+	return false;
+}
+
+bool FOnlineFriendsIOS::GetRecentPlayers(const FUniqueNetId& UserId, const FString& Namespace, TArray< TSharedRef<FOnlineRecentPlayer> >& OutRecentPlayers)
+{
+	return false;
+}
+
+bool FOnlineFriendsIOS::BlockPlayer(int32 LocalUserNum, const FUniqueNetId& PlayerId)
+{
+	return false;
+}
+
+bool FOnlineFriendsIOS::UnblockPlayer(int32 LocalUserNum, const FUniqueNetId& PlayerId)
+{
+	return false;
+}
+
+bool FOnlineFriendsIOS::QueryBlockedPlayers(const FUniqueNetId& UserId)
+{
+	return false;
+}
+
+bool FOnlineFriendsIOS::GetBlockedPlayers(const FUniqueNetId& UserId, TArray< TSharedRef<FOnlineBlockedPlayer> >& OutBlockedPlayers)
+{
+	return false;
+}
+
+void FOnlineFriendsIOS::DumpBlockedPlayers() const
+{
+}
\ No newline at end of file
diff --git a/Engine/Plugins/Online/IOS/OnlineSubsystemIOS/Source/Private/OnlineIdentityInterfaceIOS.cpp b/Engine/Plugins/Online/IOS/OnlineSubsystemIOS/Source/Private/OnlineIdentityInterfaceIOS.cpp
new file mode 100644
--- /dev/null
+++ b/Engine/Plugins/Online/IOS/OnlineSubsystemIOS/Source/Private/OnlineIdentityInterfaceIOS.cpp
@@ -0,0 +1,236 @@
+// Copyright 1998-2016 Epic Games, Inc. All Rights Reserved.
+
+#include "OnlineSubsystemIOSPrivatePCH.h"
+
+FOnlineIdentityIOS::FOnlineIdentityIOS() :
+	UniqueNetId( NULL )
+{
+}
+
+TSharedPtr<FUniqueNetIdString> FOnlineIdentityIOS::GetLocalPlayerUniqueId() const
+{
+	return UniqueNetId;
+}
+
+void FOnlineIdentityIOS::SetLocalPlayerUniqueId(const TSharedPtr<FUniqueNetIdString>& UniqueId)
+{
+	UniqueNetId = UniqueId;
+}
+
+TSharedPtr<FUserOnlineAccount> FOnlineIdentityIOS::GetUserAccount(const FUniqueNetId& UserId) const
+{
+	// not implemented
+	TSharedPtr<FUserOnlineAccount> Result;
+	return Result;
+}
+
+TArray<TSharedPtr<FUserOnlineAccount> > FOnlineIdentityIOS::GetAllUserAccounts() const
+{
+	// not implemented
+	TArray<TSharedPtr<FUserOnlineAccount> > Result;
+	return Result;
+}
+
+bool FOnlineIdentityIOS::Login(int32 LocalUserNum, const FOnlineAccountCredentials& AccountCredentials)
+{
+	bool bStartedLogin = false;
+
+	// Since the iOS login code may show a UI, ShowLoginUI is a better fit here. Also, note that the ConnectToService blueprint
+	// node that calls Login is deprecated (there's a new ShowExternalLoginUI node meant to replace it).
+	UE_LOG(LogOnline, Warning, TEXT("Using the IOnlineIdentity::Login function on iOS is not recommended. Please use IOnlineExternalUI::ShowLoginUI instead."));
+
+	// Was the login handled by Game Center
+	if( GetLocalGameCenterUser() && 
+		GetLocalGameCenterUser().isAuthenticated )
+	{
+        // Now logged in
+		bStartedLogin = true;
+        
+		const FString PlayerId(GetLocalGameCenterUser().playerID);
+		UniqueNetId = MakeShareable( new FUniqueNetIdString( PlayerId ) );
+		TriggerOnLoginCompleteDelegates(LocalUserNum, true, *UniqueNetId, TEXT(""));
+        
+        UE_LOG(LogOnline, Log, TEXT("The user %s has logged into Game Center"), *PlayerId);
+	}
+	else if([IOSAppDelegate GetDelegate].OSVersion >= 6.0f)
+	{
+		// Trigger the login event on the main thread.
+		bStartedLogin = true;
+		dispatch_async(dispatch_get_main_queue(), ^
+		{
+			[GetLocalGameCenterUser() setAuthenticateHandler:(^(UIViewController* viewcontroller, NSError *error)
+			{
+				bool bWasSuccessful = false;
+
+				// The login process has completed.
+				if (viewcontroller == nil)
+				{
+					FString ErrorMessage;
+
+					if (GetLocalGameCenterUser().isAuthenticated == YES)
+					{
+						/* Perform additional tasks for the authenticated player here */
+						const FString PlayerId(GetLocalGameCenterUser().playerID);
+						UniqueNetId = MakeShareable(new FUniqueNetIdString(PlayerId));
+
+						bWasSuccessful = true;
+						UE_LOG(LogOnline, Log, TEXT("The user %s has logged into Game Center"), *PlayerId);
+					}
+					else
+					{
+						ErrorMessage = TEXT("The user could not be authenticated by Game Center");
+						UE_LOG(LogOnline, Log, TEXT("%s"), *ErrorMessage);
+					}
+
+					if (error)
+					{
+						NSString *errstr = [error localizedDescription];
+						UE_LOG(LogOnline, Warning, TEXT("Game Center login has failed. %s]"), *FString(errstr));
+					}
+
+					// Report back to the game thread whether this succeeded.
+					[FIOSAsyncTask CreateTaskWithBlock : ^ bool(void)
+					{
+						TSharedPtr<FUniqueNetIdString> UniqueIdForUser = UniqueNetId.IsValid() ? UniqueNetId : MakeShareable(new FUniqueNetIdString());
+						TriggerOnLoginCompleteDelegates(LocalUserNum, bWasSuccessful, *UniqueIdForUser, *ErrorMessage);
+
+						return true;
+					}];
+				}
+				else
+				{
+					// Game Center has provided a view controller for us to login, we present it.
+					[[IOSAppDelegate GetDelegate].IOSController 
+						presentViewController:viewcontroller animated:YES completion:nil];
+				}
+			})];
+		});
+	}
+	else
+	{
+		// User is not currently logged into game center
+		TriggerOnLoginCompleteDelegates(LocalUserNum, false, FUniqueNetIdString(), TEXT("IOS version is not compatible with the game center implementation"));
+	}
+	
+	return bStartedLogin;
+}
+
+bool FOnlineIdentityIOS::Logout(int32 LocalUserNum)
+{
+	TriggerOnLogoutCompleteDelegates(LocalUserNum, false);
+	return true;
+}
+
+bool FOnlineIdentityIOS::AutoLogin(int32 LocalUserNum)
+{
+	return Login( LocalUserNum, FOnlineAccountCredentials() );
+}
+
+ELoginStatus::Type FOnlineIdentityIOS::GetLoginStatus(int32 LocalUserNum) const
+{
+	ELoginStatus::Type LoginStatus = ELoginStatus::NotLoggedIn;
+
+	if(LocalUserNum < MAX_LOCAL_PLAYERS && GetLocalGameCenterUser() != NULL && GetLocalGameCenterUser().isAuthenticated == YES)
+	{
+		LoginStatus = ELoginStatus::LoggedIn;
+	}
+
+	return LoginStatus;
+}
+
+ELoginStatus::Type FOnlineIdentityIOS::GetLoginStatus(const FUniqueNetId& UserId) const 
+{
+	ELoginStatus::Type LoginStatus = ELoginStatus::NotLoggedIn;
+
+	if(GetLocalGameCenterUser() != NULL && GetLocalGameCenterUser().isAuthenticated == YES)
+	{
+		LoginStatus = ELoginStatus::LoggedIn;
+	}
+
+	return LoginStatus;
+}
+
+TSharedPtr<const FUniqueNetId> FOnlineIdentityIOS::GetUniquePlayerId(int32 LocalUserNum) const
+{
+	return UniqueNetId;
+}
+
+TSharedPtr<const FUniqueNetId> FOnlineIdentityIOS::CreateUniquePlayerId(uint8* Bytes, int32 Size)
+{
+	if( Bytes && Size == sizeof(uint64) )
+	{
+		int32 StrLen = FCString::Strlen((TCHAR*)Bytes);
+		if (StrLen > 0)
+		{
+			FString StrId((TCHAR*)Bytes);
+			return MakeShareable(new FUniqueNetIdString(StrId));
+		}
+	}
+    
+	return NULL;
+}
+
+TSharedPtr<const FUniqueNetId> FOnlineIdentityIOS::CreateUniquePlayerId(const FString& Str)
+{
+	return MakeShareable(new FUniqueNetIdString(Str));
+}
+
+FString FOnlineIdentityIOS::GetPlayerNickname(int32 LocalUserNum) const
+{
+	if (LocalUserNum < MAX_LOCAL_PLAYERS && GetLocalGameCenterUser() != NULL)
+	{
+		NSString* PersonaName = [GetLocalGameCenterUser() alias];
+		
+        if (PersonaName != nil)
+        {
+            return FString(PersonaName);
+        }
+	}
+
+	return FString();
+}
+
+FString FOnlineIdentityIOS::GetPlayerNickname(const FUniqueNetId& UserId) const 
+{
+	if (GetLocalGameCenterUser() != NULL)
+	{
+		NSString* PersonaName = [GetLocalGameCenterUser() alias];
+		
+        if (PersonaName != nil)
+        {
+            return FString(PersonaName);
+        }
+	}
+
+	return FString();
+}
+
+FString FOnlineIdentityIOS::GetAuthToken(int32 LocalUserNum) const
+{
+	FString ResultToken;
+	return ResultToken;
+}
+
+void FOnlineIdentityIOS::GetUserPrivilege(const FUniqueNetId& UserId, EUserPrivileges::Type Privilege, const FOnGetUserPrivilegeCompleteDelegate& Delegate)
+{
+	Delegate.ExecuteIfBound(UserId, Privilege, (uint32)EPrivilegeResults::NoFailures);
+}
+
+FPlatformUserId FOnlineIdentityIOS::GetPlatformUserIdFromUniqueNetId(const FUniqueNetId& InUniqueNetId)
+{
+	for (int i = 0; i < MAX_LOCAL_PLAYERS; ++i)
+	{
+		auto CurrentUniqueId = GetUniquePlayerId(i);
+		if (CurrentUniqueId.IsValid() && (*CurrentUniqueId == InUniqueNetId))
+		{
+			return i;
+		}
+	}
+
+	return PLATFORMUSERID_NONE;
+}
+
+FString FOnlineIdentityIOS::GetAuthType() const
+{
+	return TEXT("");
+}
diff --git a/Engine/Plugins/Online/IOS/OnlineSubsystemIOS/Source/Private/OnlineIdentityInterfaceIOS.h b/Engine/Plugins/Online/IOS/OnlineSubsystemIOS/Source/Private/OnlineIdentityInterfaceIOS.h
new file mode 100644
--- /dev/null
+++ b/Engine/Plugins/Online/IOS/OnlineSubsystemIOS/Source/Private/OnlineIdentityInterfaceIOS.h
@@ -0,0 +1,73 @@
+// Copyright 1998-2016 Epic Games, Inc. All Rights Reserved.
+
+
+#pragma once
+
+
+#include "OnlineSubsystemIOSTypes.h"
+#include "OnlineIdentityInterface.h"
+#include "OnlineSubsystemIOSPackage.h"
+
+
+class FOnlineIdentityIOS :
+	public IOnlineIdentity
+{
+private:
+	/** UID for this identity */
+	TSharedPtr< FUniqueNetIdString > UniqueNetId;
+
+
+PACKAGE_SCOPE:
+
+	/**
+	 * Default Constructor
+	 */
+	FOnlineIdentityIOS();	
+
+	TSharedPtr<FUniqueNetIdString> GetLocalPlayerUniqueId() const;
+
+	void SetLocalPlayerUniqueId(const TSharedPtr<FUniqueNetIdString>& UniqueId);
+
+public:
+
+	//~ Begin IOnlineIdentity Interface
+	virtual bool Login(int32 LocalUserNum, const FOnlineAccountCredentials& AccountCredentials) override;
+	virtual bool Logout(int32 LocalUserNum) override;
+	virtual bool AutoLogin(int32 LocalUserNum) override;
+	virtual TSharedPtr<FUserOnlineAccount> GetUserAccount(const FUniqueNetId& UserId) const override;
+	virtual TArray<TSharedPtr<FUserOnlineAccount> > GetAllUserAccounts() const override;
+	virtual TSharedPtr<const FUniqueNetId> GetUniquePlayerId(int32 LocalUserNum) const override;
+	virtual TSharedPtr<const FUniqueNetId> CreateUniquePlayerId(uint8* Bytes, int32 Size) override;
+	virtual TSharedPtr<const FUniqueNetId> CreateUniquePlayerId(const FString& Str) override;
+	virtual ELoginStatus::Type GetLoginStatus(int32 LocalUserNum) const override;
+	virtual ELoginStatus::Type GetLoginStatus(const FUniqueNetId& UserId) const override;
+	virtual FString GetPlayerNickname(int32 LocalUserNum) const override;
+	virtual FString GetPlayerNickname(const FUniqueNetId& UserId) const override;
+	virtual FString GetAuthToken(int32 LocalUserNum) const override;
+	virtual void GetUserPrivilege(const FUniqueNetId& UserId, EUserPrivileges::Type Privilege, const FOnGetUserPrivilegeCompleteDelegate& Delegate) override;
+	virtual FPlatformUserId GetPlatformUserIdFromUniqueNetId(const FUniqueNetId& UniqueNetId) override;
+	virtual FString GetAuthType() const override;
+	//~ End IOnlineIdentity Interface
+
+
+public:
+
+	/**
+	 * Destructor
+	 */
+	virtual ~FOnlineIdentityIOS() {};
+	
+
+	/**
+	 * Get a reference to the GKLocalPlayer
+	 *
+	 * @return - The game center local player
+	 */
+	GKLocalPlayer* GetLocalGameCenterUser() const
+	{
+		return [GKLocalPlayer localPlayer];
+	}
+};
+
+
+typedef TSharedPtr<FOnlineIdentityIOS, ESPMode::ThreadSafe> FOnlineIdentityIOSPtr;
\ No newline at end of file
