Copying //UE4/GamePlugins-Staging to //UE4/Dev-Main (Source: //GamePlugins/Main @ 3232656)

diff --git a/Engine/Plugins/Online/Android/OnlineSubsystemGooglePlay/Source/Private/OnlineStoreInterfaceGooglePlay.cpp b/Engine/Plugins/Online/Android/OnlineSubsystemGooglePlay/Source/Private/OnlineStoreInterfaceGooglePlay.cpp
--- a/Engine/Plugins/Online/Android/OnlineSubsystemGooglePlay/Source/Private/OnlineStoreInterfaceGooglePlay.cpp
+++ b/Engine/Plugins/Online/Android/OnlineSubsystemGooglePlay/Source/Private/OnlineStoreInterfaceGooglePlay.cpp
@@ -1,6 +1,7 @@
 // Copyright 1998-2017 Epic Games, Inc. All Rights Reserved.
 
 #include "OnlineStoreInterfaceGooglePlay.h"
+#include "OnlinePurchaseGooglePlay.h"
 #include "OnlineAsyncTaskGooglePlayQueryInAppPurchases.h"
 #include "TaskGraphInterfaces.h"
 #include "Misc/ConfigCacheIni.h"
@@ -12,26 +13,46 @@
 ////////////////////////////////////////////////////////////////////
 /// FOnlineStoreGooglePlay implementation
 
-
 FOnlineStoreGooglePlay::FOnlineStoreGooglePlay(FOnlineSubsystemGooglePlay* InSubsystem)
 	: Subsystem( InSubsystem )
 	, CurrentQueryTask( nullptr )
 {
 	UE_LOG(LogOnline, Display, TEXT( "FOnlineStoreGooglePlay::FOnlineStoreGooglePlay" ));
-	
-	FString LicenseKey;
-	GConfig->GetString(TEXT("/Script/AndroidRuntimeSettings.AndroidRuntimeSettings"), TEXT("GooglePlayLicenseKey"), LicenseKey, GEngineIni);
-
-	extern void AndroidThunkCpp_Iap_SetupIapService(const FString&);
-	AndroidThunkCpp_Iap_SetupIapService(LicenseKey);
 }
 
-
 FOnlineStoreGooglePlay::~FOnlineStoreGooglePlay()
 {
 	UE_LOG(LogOnline, Display, TEXT( "FOnlineStoreGooglePlay::~FOnlineStoreGooglePlay" ));
+
+	if (Subsystem)
+	{
+		Subsystem->ClearOnGooglePlayAvailableIAPQueryCompleteDelegate_Handle(AvailableIAPQueryDelegateHandle);
+		Subsystem->ClearOnGooglePlayProcessPurchaseCompleteDelegate_Handle(ProcessPurchaseResultDelegateHandle);
+		Subsystem->ClearOnGooglePlayRestorePurchasesCompleteDelegate_Handle(RestorePurchasesCompleteDelegateHandle);
+	}
 }
 
+void FOnlineStoreGooglePlay::Init()
+{
+	UE_LOG(LogOnline, Display, TEXT("FOnlineStoreGooglePlay::Init"));
+	FOnGooglePlayAvailableIAPQueryCompleteDelegate IAPQueryDelegate = FOnGooglePlayAvailableIAPQueryCompleteDelegate::CreateThreadSafeSP(this, &FOnlineStoreGooglePlay::OnGooglePlayAvailableIAPQueryComplete);
+	AvailableIAPQueryDelegateHandle = Subsystem->AddOnGooglePlayAvailableIAPQueryCompleteDelegate_Handle(IAPQueryDelegate);
+
+	FOnGooglePlayProcessPurchaseCompleteDelegate PurchaseCompleteDelegate = FOnGooglePlayProcessPurchaseCompleteDelegate::CreateThreadSafeSP(this, &FOnlineStoreGooglePlay::OnProcessPurchaseResult);
+	ProcessPurchaseResultDelegateHandle = Subsystem->AddOnGooglePlayProcessPurchaseCompleteDelegate_Handle(PurchaseCompleteDelegate);
+
+	FOnGooglePlayRestorePurchasesCompleteDelegate RestorePurchasesCompleteDelegate = FOnGooglePlayRestorePurchasesCompleteDelegate::CreateThreadSafeSP(this, &FOnlineStoreGooglePlay::OnRestorePurchasesComplete);
+	RestorePurchasesCompleteDelegateHandle = Subsystem->AddOnGooglePlayRestorePurchasesCompleteDelegate_Handle(RestorePurchasesCompleteDelegate);
+
+	FString GooglePlayLicenseKey;
+	if (!GConfig->GetString(TEXT("/Script/AndroidRuntimeSettings.AndroidRuntimeSettings"), TEXT("GooglePlayLicenseKey"), GooglePlayLicenseKey, GEngineIni) || GooglePlayLicenseKey.IsEmpty())
+	{
+		UE_LOG_ONLINE(Warning, TEXT("Missing GooglePlayLicenseKey key in /Script/AndroidRuntimeSettings.AndroidRuntimeSettings of DefaultEngine.ini"));
+	}
+
+	extern void AndroidThunkCpp_Iap_SetupIapService(const FString&);
+	AndroidThunkCpp_Iap_SetupIapService(GooglePlayLicenseKey);
+}
 
 bool FOnlineStoreGooglePlay::IsAllowedToMakePurchases()
 {
@@ -49,13 +70,9 @@ bool FOnlineStoreGooglePlay::QueryForAvailablePurchases(const TArray<FString>& P
 	ReadObject = InReadObject;
 	ReadObject->ReadState = EOnlineAsyncTaskState::InProgress;
 
-	TArray<bool> ConsumableFlags;
-	ConsumableFlags.AddZeroed(ProductIds.Num());
-
 	CurrentQueryTask = new FOnlineAsyncTaskGooglePlayQueryInAppPurchases(
 		Subsystem,
-		ProductIds,
-		ConsumableFlags);
+		ProductIds);
 	Subsystem->QueueAsyncTask(CurrentQueryTask);
 
 	return true;
@@ -133,29 +150,32 @@ extern "C" void Java_com_epicgames_ue4_GooglePlayStoreHelper_nativeQueryComplete
 
 	DECLARE_CYCLE_STAT(TEXT("FSimpleDelegateGraphTask.ProcessQueryIapResult"), STAT_FSimpleDelegateGraphTask_ProcessQueryIapResult, STATGROUP_TaskGraphTasks);
 
+	FPlatformMisc::LowLevelOutputDebugStringf(TEXT("Adding task Success: %d Response: %s"), bWasSuccessful, ToString(EGPResponse));
+
 	FSimpleDelegateGraphTask::CreateAndDispatchWhenReady(
-		FSimpleDelegateGraphTask::FDelegate::CreateLambda([=](){
-			if (IOnlineSubsystem* const OnlineSub = IOnlineSubsystem::Get())
+		FSimpleDelegateGraphTask::FDelegate::CreateLambda([=]()
+	{
+		if (IOnlineSubsystem* const OnlineSub = IOnlineSubsystem::Get(GOOGLEPLAY_SUBSYSTEM))
+		{
+			FOnlineSubsystemGooglePlay* const OnlineSubGP = static_cast<FOnlineSubsystemGooglePlay* const>(OnlineSub);
+			if (OnlineSubGP)
 			{
-				// call store implementation to process query results.
-				if (FOnlineStoreGooglePlay* StoreInterface = (FOnlineStoreGooglePlay*)OnlineSub->GetStoreInterface().Get())
-				{
-					StoreInterface->ProcessQueryAvailablePurchasesResults(EGPResponse, ProvidedProductInformation);
-				}
+				FPlatformMisc::LowLevelOutputDebugStringf(TEXT("TriggerOnGooglePlayAvailableIAPQueryCompleteDelegates %s Size: %d"), ToString(EGPResponse), ProvidedProductInformation.Num());
+				OnlineSubGP->TriggerOnGooglePlayAvailableIAPQueryCompleteDelegates(EGPResponse, ProvidedProductInformation);
 			}
-			FPlatformMisc::LowLevelOutputDebugStringf(TEXT("In-App Purchase query was completed  %s\n"), bWasSuccessful ? TEXT("successfully") : TEXT("unsuccessfully"));
-		}),
-		GET_STATID(STAT_FSimpleDelegateGraphTask_ProcessQueryIapResult), 
-		nullptr, 
+		}
+		FPlatformMisc::LowLevelOutputDebugStringf(TEXT("In-App Purchase query was completed  %s\n"), bWasSuccessful ? TEXT("successfully") : TEXT("unsuccessfully"));
+	}),
+		GET_STATID(STAT_FSimpleDelegateGraphTask_ProcessQueryIapResult),
+		nullptr,
 		ENamedThreads::GameThread
-	);
-
+		);
 }
 
 
-void FOnlineStoreGooglePlay::ProcessQueryAvailablePurchasesResults(EGooglePlayBillingResponseCode InResponseCode, const TArray<FInAppPurchaseProductInfo>& AvailablePurchases)
+void FOnlineStoreGooglePlay::OnGooglePlayAvailableIAPQueryComplete(EGooglePlayBillingResponseCode InResponseCode, const TArray<FInAppPurchaseProductInfo>& AvailablePurchases)
 {
-	UE_LOG(LogOnline, Display, TEXT("FOnlineStoreGooglePlay::ProcessQueryAvailablePurchasesResults"));
+	UE_LOG(LogOnline, Display, TEXT("FOnlineStoreGooglePlay::OnGooglePlayAvailableIAPQueryComplete"));
 
 	if (ReadObject.IsValid())
 	{
@@ -164,8 +184,10 @@ void FOnlineStoreGooglePlay::ProcessQueryAvailablePurchasesResults(EGooglePlayBi
 	}
 
 	CurrentQueryTask->ProcessQueryAvailablePurchasesResults(InResponseCode == EGooglePlayBillingResponseCode::Ok);
-}
 
+	// clear the pointer, it will be destroyed by the async task manager
+	CurrentQueryTask = nullptr;
+}
 
 bool FOnlineStoreGooglePlay::BeginPurchase(const FInAppPurchaseProductRequest& ProductRequest, FOnlineInAppPurchaseTransactionRef& InPurchaseStateObject)
 {
@@ -176,9 +198,10 @@ bool FOnlineStoreGooglePlay::BeginPurchase(const FInAppPurchaseProductRequest& P
 	if (IsAllowedToMakePurchases())
 	{
 		CachedPurchaseStateObject = InPurchaseStateObject;
+		CachedPurchaseStateObject->bIsConsumable = ProductRequest.bIsConsumable;
 
-		extern bool AndroidThunkCpp_Iap_BeginPurchase(const FString&, const bool);
-		bCreatedNewTransaction = AndroidThunkCpp_Iap_BeginPurchase(ProductRequest.ProductIdentifier, ProductRequest.bIsConsumable);
+		extern bool AndroidThunkCpp_Iap_BeginPurchase(const FString&);
+		bCreatedNewTransaction = AndroidThunkCpp_Iap_BeginPurchase(ProductRequest.ProductIdentifier);
 		UE_LOG(LogOnline, Display, TEXT("Created Transaction? - %s"), 
 			bCreatedNewTransaction ? TEXT("Created a transaction.") : TEXT("Failed to create a transaction."));
 
@@ -201,23 +224,26 @@ bool FOnlineStoreGooglePlay::BeginPurchase(const FInAppPurchaseProductRequest& P
 		TriggerOnInAppPurchaseCompleteDelegates(EInAppPurchaseState::NotAllowed);
 	}
 
-
 	return bCreatedNewTransaction;
 }
 
 
-extern "C" void Java_com_epicgames_ue4_GooglePlayStoreHelper_nativePurchaseComplete(JNIEnv* jenv, jobject thiz, jsize responseCode, jstring productId, jstring receiptData, jstring signature)
+extern "C" void Java_com_epicgames_ue4_GooglePlayStoreHelper_nativePurchaseComplete(JNIEnv* jenv, jobject thiz, jsize responseCode, jstring productId, jstring productToken, jstring receiptData, jstring signature)
 {
-	FString ProductId, ReceiptData, Signature;
+	FString ProductId, ProductToken, ReceiptData, Signature;
 	EGooglePlayBillingResponseCode EGPResponse = (EGooglePlayBillingResponseCode)responseCode;
-	bool bWasSuccessful = (EGPResponse == EGooglePlayBillingResponseCode::Ok);
 
+	bool bWasSuccessful = (EGPResponse == EGooglePlayBillingResponseCode::Ok);
 	if (bWasSuccessful)
 	{
 		const char* charsId = jenv->GetStringUTFChars(productId, 0);
 		ProductId = FString(UTF8_TO_TCHAR(charsId));
 		jenv->ReleaseStringUTFChars(productId, charsId);
 
+		const char* charsToken = jenv->GetStringUTFChars(productToken, 0);
+		ProductToken = FString(UTF8_TO_TCHAR(charsToken));
+		jenv->ReleaseStringUTFChars(productToken, charsToken);
+
 		const char* charsReceipt = jenv->GetStringUTFChars(receiptData, 0);
 		ReceiptData = FString(UTF8_TO_TCHAR(charsReceipt));
 		jenv->ReleaseStringUTFChars(receiptData, charsReceipt);
@@ -226,48 +252,60 @@ extern "C" void Java_com_epicgames_ue4_GooglePlayStoreHelper_nativePurchaseCompl
 		Signature = FString(UTF8_TO_TCHAR(charsSignature));
 		jenv->ReleaseStringUTFChars(signature, charsSignature);
 	}
-	FPlatformMisc::LowLevelOutputDebugStringf(TEXT("1... ProductId: %s, ReceiptData: %s, Signature: %s\n"), *ProductId, *ReceiptData, *Signature );
+
+	FGoogleTransactionData TransactionData(ProductId, ProductToken, ReceiptData, Signature);
+
+	FPlatformMisc::LowLevelOutputDebugStringf(TEXT("1... Response: %s, Transaction %s"), ToString(EGPResponse), *TransactionData.ToDebugString());
 
 	DECLARE_CYCLE_STAT(TEXT("FSimpleDelegateGraphTask.ProcessIapResult"), STAT_FSimpleDelegateGraphTask_ProcessIapResult, STATGROUP_TaskGraphTasks);
 
 	FSimpleDelegateGraphTask::CreateAndDispatchWhenReady(
-		FSimpleDelegateGraphTask::FDelegate::CreateLambda([=](){
-			FPlatformMisc::LowLevelOutputDebugStringf(TEXT("In-App Purchase was completed  %s\n"), bWasSuccessful ? TEXT("successfully") : TEXT("unsuccessfully"));
-			if (IOnlineSubsystem* const OnlineSub = IOnlineSubsystem::Get())
+		FSimpleDelegateGraphTask::FDelegate::CreateLambda([=]()
+	{
+		FPlatformMisc::LowLevelOutputDebugStringf(TEXT("In-App Purchase was completed  %s\n"), bWasSuccessful ? TEXT("successfully") : TEXT("unsuccessfully"));
+		if (IOnlineSubsystem* const OnlineSub = IOnlineSubsystem::Get(GOOGLEPLAY_SUBSYSTEM))
+		{
+			FOnlineSubsystemGooglePlay* const OnlineSubGP = static_cast<FOnlineSubsystemGooglePlay* const>(OnlineSub);
+			if (OnlineSubGP)
 			{
-				FPlatformMisc::LowLevelOutputDebugStringf(TEXT("2... ProductId: %s, ReceiptData: %s, Signature: %s\n"), *ProductId, *ReceiptData, *Signature);
-				// call store implementation to process query results.
-				if (FOnlineStoreGooglePlay* StoreInterface = (FOnlineStoreGooglePlay*)OnlineSub->GetStoreInterface().Get())
-				{
-					StoreInterface->ProcessPurchaseResult(EGPResponse, ProductId, ReceiptData, Signature);
-				}
+				FPlatformMisc::LowLevelOutputDebugStringf(TEXT("2... Response %s Transaction %s"), ToString(EGPResponse), *TransactionData.ToDebugString());
+				OnlineSubGP->TriggerOnGooglePlayProcessPurchaseCompleteDelegates(EGPResponse, TransactionData);
 			}
-		}),
-		GET_STATID(STAT_FSimpleDelegateGraphTask_ProcessIapResult), 
-		nullptr, 
+		}
+	}),
+		GET_STATID(STAT_FSimpleDelegateGraphTask_ProcessIapResult),
+		nullptr,
 		ENamedThreads::GameThread
-	);
-
+		);
 }
 
 
-void FOnlineStoreGooglePlay::ProcessPurchaseResult(EGooglePlayBillingResponseCode InResponseCode, const FString& ProductId, const FString& InReceiptData, const FString& Signature)
+void FOnlineStoreGooglePlay::OnProcessPurchaseResult(EGooglePlayBillingResponseCode InResponseCode, const FGoogleTransactionData& InTransactionData)
 {
-	UE_LOG(LogOnline, Display, TEXT("FOnlineStoreGooglePlay::ProcessPurchaseResult"));
-	UE_LOG(LogOnline, Display, TEXT("3... ProductId: %s, ReceiptData: %s, Signature: %s\n"), *ProductId, *InReceiptData, *Signature);
-
+	UE_LOG(LogOnline, Display, TEXT("FOnlineStoreGooglePlay::OnProcessPurchaseResult"));
+	UE_LOG(LogOnline, Display, TEXT("3... Response: %s Transaction: %s"), ToString(InResponseCode), *InTransactionData.ToDebugString());
 
+	bool bWasSuccessful = (InResponseCode == EGooglePlayBillingResponseCode::Ok);
 	if (CachedPurchaseStateObject.IsValid())
 	{
+		if (CachedPurchaseStateObject->bIsConsumable && InTransactionData.GetErrorStr().IsEmpty())
+		{
+			// Consume right away to maintain behavior of legacy code (GooglePlayStoreHelper.java)
+			// Technically wrong/dangerous because the game may not grant entitlements
+			// on a crash or other interruption 
+			extern void AndroidThunkCpp_Iap_ConsumePurchase(const FString&);
+			AndroidThunkCpp_Iap_ConsumePurchase(InTransactionData.GetTransactionIdentifier());
+		}
+
 		FInAppPurchaseProductInfo& ProductInfo = CachedPurchaseStateObject->ProvidedProductInformation;
-		ProductInfo.Identifier = ProductId;
+		ProductInfo.Identifier = InTransactionData.GetOfferId();
 		ProductInfo.DisplayName = TEXT("n/a");
 		ProductInfo.DisplayDescription = TEXT("n/a");
 		ProductInfo.DisplayPrice = TEXT("n/a");
-		ProductInfo.ReceiptData = InReceiptData;
-		ProductInfo.TransactionIdentifier = Signature;
+		ProductInfo.ReceiptData = InTransactionData.GetCombinedReceiptData();
+		ProductInfo.TransactionIdentifier = InTransactionData.GetTransactionIdentifier();
 
-		CachedPurchaseStateObject->ReadState = EOnlineAsyncTaskState::Done;
+		CachedPurchaseStateObject->ReadState = bWasSuccessful ? EOnlineAsyncTaskState::Done : EOnlineAsyncTaskState::Failed;
 	}
 
 	TriggerOnInAppPurchaseCompleteDelegates(ConvertGPResponseCodeToIAPState(InResponseCode));
@@ -303,9 +341,9 @@ bool FOnlineStoreGooglePlay::RestorePurchases(const TArray<FInAppPurchaseProduct
 	return bSentAQueryRequest;
 }
 
-extern "C" void Java_com_epicgames_ue4_GooglePlayStoreHelper_nativeRestorePurchasesComplete(JNIEnv* jenv, jobject thiz, jsize responseCode, jobjectArray ProductIDs, jobjectArray ReceiptsData, jobjectArray Signatures)
+extern "C" void Java_com_epicgames_ue4_GooglePlayStoreHelper_nativeRestorePurchasesComplete(JNIEnv* jenv, jobject thiz, jsize responseCode, jobjectArray ProductIDs, jobjectArray ProductTokens, jobjectArray ReceiptsData, jobjectArray Signatures)
 {
-	TArray<FInAppPurchaseRestoreInfo> RestoredPurchaseInfo;
+	TArray<FGoogleTransactionData> RestoredPurchaseInfo;
 
 	EGooglePlayBillingResponseCode EGPResponse = (EGooglePlayBillingResponseCode)responseCode;
 	bool bWasSuccessful = (EGPResponse == EGooglePlayBillingResponseCode::Ok);
@@ -313,10 +351,11 @@ extern "C" void Java_com_epicgames_ue4_GooglePlayStoreHelper_nativeRestorePurcha
 	if (jenv && bWasSuccessful)
 	{
 		jsize NumProducts = jenv->GetArrayLength(ProductIDs);
+		jsize NumProductTokens = jenv->GetArrayLength(ProductTokens);
 		jsize NumReceipts = jenv->GetArrayLength(ReceiptsData);
 		jsize NumSignatures = jenv->GetArrayLength(Signatures);
 
-		ensure((NumProducts == NumReceipts) && (NumProducts == NumSignatures));
+		ensure((NumProducts == NumProductTokens) && (NumProducts == NumReceipts) && (NumProducts == NumSignatures));
 
 		for (jsize Idx = 0; Idx < NumProducts; Idx++)
 		{
@@ -325,25 +364,32 @@ extern "C" void Java_com_epicgames_ue4_GooglePlayStoreHelper_nativeRestorePurcha
 
 			jstring NextId = (jstring)jenv->GetObjectArrayElement(ProductIDs, Idx);
 			const char* charsId = jenv->GetStringUTFChars(NextId, 0);
-			RestoreInfo.Identifier = FString(UTF8_TO_TCHAR(charsId));
+			const FString OfferId = FString(UTF8_TO_TCHAR(charsId));
 			jenv->ReleaseStringUTFChars(NextId, charsId);
 			jenv->DeleteLocalRef(NextId);
 
+			jstring NextToken = (jstring)jenv->GetObjectArrayElement(ProductTokens, Idx);
+			const char* charsToken = jenv->GetStringUTFChars(NextToken, 0);
+			const FString ProductToken = FString(UTF8_TO_TCHAR(charsToken));
+			jenv->ReleaseStringUTFChars(NextToken, charsToken);
+			jenv->DeleteLocalRef(NextToken);
+
 			jstring NextReceipt = (jstring)jenv->GetObjectArrayElement(ReceiptsData, Idx);
 			const char* charsReceipt = jenv->GetStringUTFChars(NextReceipt, 0);
-			RestoreInfo.ReceiptData = FString(UTF8_TO_TCHAR(charsReceipt));
+			const FString ReceiptData = FString(UTF8_TO_TCHAR(charsReceipt));
 			jenv->ReleaseStringUTFChars(NextReceipt, charsReceipt);
 			jenv->DeleteLocalRef(NextReceipt);
 
 			jstring NextSignature = (jstring)jenv->GetObjectArrayElement(Signatures, Idx);
 			const char* charsSignature = jenv->GetStringUTFChars(NextSignature, 0);
-			RestoreInfo.TransactionIdentifier = FString(UTF8_TO_TCHAR(charsSignature));
+			const FString SignatureData = FString(UTF8_TO_TCHAR(charsSignature));
 			jenv->ReleaseStringUTFChars(NextSignature, charsSignature);
 			jenv->DeleteLocalRef(NextSignature);
 
-			RestoredPurchaseInfo.Add(RestoreInfo);
+			FGoogleTransactionData RestoredPurchase(OfferId, ProductToken, ReceiptData, SignatureData);
+			RestoredPurchaseInfo.Add(RestoredPurchase);
 
-			FPlatformMisc::LowLevelOutputDebugStringf(TEXT("\nRestored Product Identifier: %s\n"), *RestoreInfo.Identifier);
+			FPlatformMisc::LowLevelOutputDebugStringf(TEXT("Restored Transaction: %s"), *RestoredPurchase.ToDebugString());
 		}
 	}
 
@@ -351,51 +397,118 @@ extern "C" void Java_com_epicgames_ue4_GooglePlayStoreHelper_nativeRestorePurcha
 
 	FSimpleDelegateGraphTask::CreateAndDispatchWhenReady(
 		FSimpleDelegateGraphTask::FDelegate::CreateLambda([=]()
-		{
-			FPlatformMisc::LowLevelOutputDebugStringf(TEXT("Restoring In-App Purchases was completed  %s\n"), bWasSuccessful ? TEXT("successfully") : TEXT("unsuccessfully"));
+	{
+		FPlatformMisc::LowLevelOutputDebugStringf(TEXT("Restoring In-App Purchases was completed  %s\n"), bWasSuccessful ? TEXT("successfully") : TEXT("unsuccessfully"));
 
-			if (IOnlineSubsystem* const OnlineSub = IOnlineSubsystem::Get())
+		if (IOnlineSubsystem* const OnlineSub = IOnlineSubsystem::Get(GOOGLEPLAY_SUBSYSTEM))
+		{
+			FOnlineSubsystemGooglePlay* const OnlineSubGP = static_cast<FOnlineSubsystemGooglePlay* const>(OnlineSub);
+			if (OnlineSubGP)
 			{
-				FPlatformMisc::LowLevelOutputDebugStringf(TEXT("Sending result back to OnlineSubsystem.\n"));
-				// call store implementation to process query results.
-				if (FOnlineStoreGooglePlay* StoreInterface = (FOnlineStoreGooglePlay*)OnlineSub->GetStoreInterface().Get())
-				{
-					if (StoreInterface->CachedPurchaseRestoreObject.IsValid())
-					{
-						StoreInterface->CachedPurchaseRestoreObject->ProvidedRestoreInformation = RestoredPurchaseInfo;
-						StoreInterface->CachedPurchaseRestoreObject->ReadState = bWasSuccessful ? EOnlineAsyncTaskState::Done : EOnlineAsyncTaskState::Failed;
-					}
-
-					EInAppPurchaseState::Type IAPState = bWasSuccessful ? EInAppPurchaseState::Restored : StoreInterface->ConvertGPResponseCodeToIAPState(EGPResponse);
-
-					StoreInterface->TriggerOnInAppPurchaseRestoreCompleteDelegates(IAPState);
-				}
+				OnlineSubGP->TriggerOnGooglePlayRestorePurchasesCompleteDelegates(EGPResponse, RestoredPurchaseInfo);
 			}
-		}),
+		}
+	}),
 		GET_STATID(STAT_FSimpleDelegateGraphTask_RestorePurchases),
 		nullptr,
 		ENamedThreads::GameThread
-	);
+		);
+}
+
+void FOnlineStoreGooglePlay::OnRestorePurchasesComplete(EGooglePlayBillingResponseCode InResponseCode, const TArray<FGoogleTransactionData>& InRestoredPurchases)
+{
+	UE_LOG(LogOnline, Verbose, TEXT("FOnlineStoreGooglePlay::OnRestorePurchasesComplete Response: %s Num: %d"), ToString(InResponseCode), InRestoredPurchases.Num());
+
+	bool bWasSuccessful = (InResponseCode == EGooglePlayBillingResponseCode::Ok);
+	if (CachedPurchaseRestoreObject.IsValid())
+	{
+		TArray<FInAppPurchaseRestoreInfo> RestoredPurchaseInfo;
+		RestoredPurchaseInfo.Reserve(InRestoredPurchases.Num());
+		for (const FGoogleTransactionData& RestoredPurchase : InRestoredPurchases)
+		{
+			FInAppPurchaseRestoreInfo RestoreInfo;
+			RestoreInfo.Identifier = RestoredPurchase.GetOfferId();
+			RestoreInfo.ReceiptData = RestoredPurchase.GetCombinedReceiptData();
+			RestoreInfo.TransactionIdentifier = RestoredPurchase.GetTransactionIdentifier();
+			RestoredPurchaseInfo.Add(RestoreInfo);
+		}
+
+		CachedPurchaseRestoreObject->ProvidedRestoreInformation = RestoredPurchaseInfo;
+		CachedPurchaseRestoreObject->ReadState = bWasSuccessful ? EOnlineAsyncTaskState::Done : EOnlineAsyncTaskState::Failed;
+	}
+
+	EInAppPurchaseState::Type IAPState = bWasSuccessful ? EInAppPurchaseState::Restored : ConvertGPResponseCodeToIAPState(InResponseCode);
+
+	TriggerOnInAppPurchaseRestoreCompleteDelegates(IAPState);
 }
 
-EInAppPurchaseState::Type FOnlineStoreGooglePlay::ConvertGPResponseCodeToIAPState(const EGooglePlayBillingResponseCode InResponseCode)
+extern "C" void Java_com_epicgames_ue4_GooglePlayStoreHelper_nativeQueryExistingPurchasesComplete(JNIEnv* jenv, jobject thiz, jsize responseCode, jobjectArray ProductIDs, jobjectArray ProductTokens, jobjectArray ReceiptsData, jobjectArray Signatures)
 {
-	switch (InResponseCode)
+	TArray<FGoogleTransactionData> ExistingPurchaseInfo;
+
+	EGooglePlayBillingResponseCode EGPResponse = (EGooglePlayBillingResponseCode)responseCode;
+
+	bool bWasSuccessful = (EGPResponse == EGooglePlayBillingResponseCode::Ok);
+	if (jenv && bWasSuccessful)
 	{
-	case EGooglePlayBillingResponseCode::Ok:
-		return EInAppPurchaseState::Success;
-	case EGooglePlayBillingResponseCode::UserCancelled:
-		return EInAppPurchaseState::Cancelled;
-	case EGooglePlayBillingResponseCode::ItemAlreadyOwned:
-		return EInAppPurchaseState::AlreadyOwned;
-	case EGooglePlayBillingResponseCode::ItemNotOwned:
-		return EInAppPurchaseState::NotAllowed;
-	case EGooglePlayBillingResponseCode::ServiceUnavailable:
-	case EGooglePlayBillingResponseCode::BillingUnavailable:
-	case EGooglePlayBillingResponseCode::ItemUnavailable:
-	case EGooglePlayBillingResponseCode::DeveloperError:
-	case EGooglePlayBillingResponseCode::Error:
-	default:
-		return EInAppPurchaseState::Failed;
+		jsize NumProducts = jenv->GetArrayLength(ProductIDs);
+		jsize NumProductTokens = jenv->GetArrayLength(ProductTokens);
+		jsize NumReceipts = jenv->GetArrayLength(ReceiptsData);
+		jsize NumSignatures = jenv->GetArrayLength(Signatures);
+
+		ensure((NumProducts == NumProductTokens) && (NumProducts == NumReceipts) && (NumProducts == NumSignatures));
+
+		for (jsize Idx = 0; Idx < NumProducts; Idx++)
+		{
+			// Build the product information strings.
+			jstring NextId = (jstring)jenv->GetObjectArrayElement(ProductIDs, Idx);
+			const char* charsId = jenv->GetStringUTFChars(NextId, 0);
+			const FString OfferId = FString(UTF8_TO_TCHAR(charsId));
+			jenv->ReleaseStringUTFChars(NextId, charsId);
+			jenv->DeleteLocalRef(NextId);
+
+			jstring NextToken = (jstring)jenv->GetObjectArrayElement(ProductTokens, Idx);
+			const char* charsToken = jenv->GetStringUTFChars(NextToken, 0);
+			const FString ProductToken = FString(UTF8_TO_TCHAR(charsToken));
+			jenv->ReleaseStringUTFChars(NextToken, charsToken);
+			jenv->DeleteLocalRef(NextToken);
+
+			jstring NextReceipt = (jstring)jenv->GetObjectArrayElement(ReceiptsData, Idx);
+			const char* charsReceipt = jenv->GetStringUTFChars(NextReceipt, 0);
+			const FString ReceiptData = FString(UTF8_TO_TCHAR(charsReceipt));
+			jenv->ReleaseStringUTFChars(NextReceipt, charsReceipt);
+			jenv->DeleteLocalRef(NextReceipt);
+
+			jstring NextSignature = (jstring)jenv->GetObjectArrayElement(Signatures, Idx);
+			const char* charsSignature = jenv->GetStringUTFChars(NextSignature, 0);
+			const FString SignatureData = FString(UTF8_TO_TCHAR(charsSignature));
+			jenv->ReleaseStringUTFChars(NextSignature, charsSignature);
+			jenv->DeleteLocalRef(NextSignature);
+
+			FGoogleTransactionData ExistingPurchase(OfferId, ProductToken, ReceiptData, SignatureData);
+			ExistingPurchaseInfo.Add(ExistingPurchase);
+
+			FPlatformMisc::LowLevelOutputDebugStringf(TEXT("\Existing Product Identifier: %s"), *ExistingPurchase.ToDebugString());
+		}
 	}
+
+	DECLARE_CYCLE_STAT(TEXT("FSimpleDelegateGraphTask.QueryExistingPurchases"), STAT_FSimpleDelegateGraphTask_QueryExistingPurchases, STATGROUP_TaskGraphTasks);
+
+	FSimpleDelegateGraphTask::CreateAndDispatchWhenReady(
+		FSimpleDelegateGraphTask::FDelegate::CreateLambda([=]()
+	{
+		FPlatformMisc::LowLevelOutputDebugStringf(TEXT("Query existing purchases was completed %s\n"), bWasSuccessful ? TEXT("successfully") : TEXT("unsuccessfully"));
+		if (IOnlineSubsystem* const OnlineSub = IOnlineSubsystem::Get(GOOGLEPLAY_SUBSYSTEM))
+		{
+			FOnlineSubsystemGooglePlay* const OnlineSubGP = static_cast<FOnlineSubsystemGooglePlay* const>(OnlineSub);
+			if (OnlineSubGP)
+			{
+				OnlineSubGP->TriggerOnGooglePlayQueryExistingPurchasesCompleteDelegates(EGPResponse, ExistingPurchaseInfo);
+			}
+		}
+	}),
+		GET_STATID(STAT_FSimpleDelegateGraphTask_QueryExistingPurchases),
+		nullptr,
+		ENamedThreads::GameThread
+		);
 }
diff --git a/Engine/Plugins/Online/Android/OnlineSubsystemGooglePlay/Source/Private/OnlineStoreInterfaceGooglePlay.h b/Engine/Plugins/Online/Android/OnlineSubsystemGooglePlay/Source/Private/OnlineStoreInterfaceGooglePlay.h
--- a/Engine/Plugins/Online/Android/OnlineSubsystemGooglePlay/Source/Private/OnlineStoreInterfaceGooglePlay.h
+++ b/Engine/Plugins/Online/Android/OnlineSubsystemGooglePlay/Source/Private/OnlineStoreInterfaceGooglePlay.h
@@ -4,56 +4,17 @@
 
 #include "CoreMinimal.h"
 #include "OnlineStoreInterface.h"
+#include "OnlineStoreGooglePlayCommon.h"
 
 class FOnlineSubsystemGooglePlay;
-
-enum class EGooglePlayBillingResponseCode
-{
-	Ok					= 0,
-	UserCancelled		= 1,
-	ServiceUnavailable	= 2,
-	BillingUnavailable	= 3,
-	ItemUnavailable		= 4,
-	DeveloperError		= 5,
-	Error				= 6,
-	ItemAlreadyOwned	= 7,
-	ItemNotOwned		= 8,
-};
-
-
-/**
- * The resulting state of an iap transaction
- */
-namespace EInAppPurchaseResult
-{
-	enum Type
-	{
-		Succeeded = 0,
-		RestoredFromServer,
-		Failed,
-		Cancelled,
-	};
-}
-
-
-/**
- * Implementation of the Platform Purchase receipt. For this we provide an identifier and the encrypted data.
- */
-class FGooglePlayPurchaseReceipt : public IPlatformPurchaseReceipt
-{
-public:
-	// Product identifier
-	FString Identifier;
-
-	// The encrypted receipt data
-	FString Data;
-};
-
+class FOnlineAsyncTaskGooglePlayQueryInAppPurchases;
 
 /**
  *	FOnlineStoreGooglePlay - Implementation of the online store for GooglePlay
  */
-class FOnlineStoreGooglePlay : public IOnlineStore
+class FOnlineStoreGooglePlay : 
+	public IOnlineStore,
+	public TSharedFromThis<FOnlineStoreGooglePlay, ESPMode::ThreadSafe>
 {
 public:
 	/** C-tor */
@@ -61,6 +22,9 @@ public:
 	/** Destructor */
 	virtual ~FOnlineStoreGooglePlay();
 
+	/** Initialize the interface */
+	void Init();
+
 	// Begin IOnlineStore 
 	virtual bool QueryForAvailablePurchases(const TArray<FString>& ProductIDs, FOnlineProductInformationReadRef& InReadObject) override;
 	virtual bool BeginPurchase(const FInAppPurchaseProductRequest& ProductRequest, FOnlineInAppPurchaseTransactionRef& InReadObject) override;
@@ -68,11 +32,6 @@ public:
 	virtual bool RestorePurchases(const TArray<FInAppPurchaseProductRequest>& ConsumableProductFlags, FOnlineInAppPurchaseRestoreReadRef& InReadObject) override;
 	// End IOnlineStore 
 
-	void ProcessQueryAvailablePurchasesResults(EGooglePlayBillingResponseCode InResponseCode, const TArray<FInAppPurchaseProductInfo>& AvailablePurchases);
-	void ProcessPurchaseResult(EGooglePlayBillingResponseCode InResponseCode, const FString& InProductId, const FString& InReceiptData, const FString& Signature);
-
-	EInAppPurchaseState::Type ConvertGPResponseCodeToIAPState(const EGooglePlayBillingResponseCode InResponseCode);
-
 	/** Cached in-app purchase restore transaction object, used to provide details to the developer about what products should be restored */
 	FOnlineInAppPurchaseRestoreReadPtr CachedPurchaseRestoreObject;
 
@@ -82,7 +41,7 @@ private:
 	FOnlineSubsystemGooglePlay* Subsystem;
 
 	/** The current query for iap async task */
-	class FOnlineAsyncTaskGooglePlayQueryInAppPurchases* CurrentQueryTask;
+	FOnlineAsyncTaskGooglePlayQueryInAppPurchases* CurrentQueryTask;
 
 	/** Delegate fired when a query for purchases has completed, whether successful or unsuccessful */
 	FOnQueryForAvailablePurchasesComplete OnQueryForAvailablePurchasesCompleteDelegate;
@@ -95,6 +54,33 @@ private:
 
 	/** Cached in-app purchase transaction object, used to provide details to the user, of the product that has just been purchased. */
 	FOnlineInAppPurchaseTransactionPtr CachedPurchaseStateObject;
+
+	/**
+	 * Delegate fired where an IAP query for available offers has completed
+	 *
+	 * @param InResponseCode response from Google backend
+	 * @param AvailablePurchases list of offers returned in response to a query on available offer ids
+	 */
+	void OnGooglePlayAvailableIAPQueryComplete(EGooglePlayBillingResponseCode InResponseCode, const TArray<FInAppPurchaseProductInfo>& AvailablePurchases);
+	FDelegateHandle AvailableIAPQueryDelegateHandle;
+
+	/**
+	 * Delegate fired when a purchase has completed
+	 *
+	 * @param InResponseCode response from the GooglePlay backend
+	 * @param InTransactionData transaction data for the completed purchase
+	 */
+	void OnProcessPurchaseResult(EGooglePlayBillingResponseCode InResponseCode, const FGoogleTransactionData& InTransactionData);
+	FDelegateHandle ProcessPurchaseResultDelegateHandle;
+
+	/**
+	 * Delegate fired when purchases are restored
+	 *
+	 * @param InResponseCode response from the GooglePlay backend
+	 * @param InRestoredPurchases transaction data for the restored purchases
+	 */
+	void OnRestorePurchasesComplete(EGooglePlayBillingResponseCode InResponseCode, const TArray<FGoogleTransactionData>& InRestoredPurchases);
+	FDelegateHandle RestorePurchasesCompleteDelegateHandle;
 };
 
 typedef TSharedPtr<FOnlineStoreGooglePlay, ESPMode::ThreadSafe> FOnlineStoreGooglePlayPtr;
diff --git a/Engine/Plugins/Online/Android/OnlineSubsystemGooglePlay/Source/Private/OnlineSubsystemGooglePlay.cpp b/Engine/Plugins/Online/Android/OnlineSubsystemGooglePlay/Source/Private/OnlineSubsystemGooglePlay.cpp
--- a/Engine/Plugins/Online/Android/OnlineSubsystemGooglePlay/Source/Private/OnlineSubsystemGooglePlay.cpp
+++ b/Engine/Plugins/Online/Android/OnlineSubsystemGooglePlay/Source/Private/OnlineSubsystemGooglePlay.cpp
@@ -3,6 +3,13 @@
 //Google Play Services
 
 #include "OnlineSubsystemGooglePlay.h"
+#include "OnlineIdentityInterfaceGooglePlay.h"
+#include "OnlineStoreGooglePlayCommon.h"
+#include "OnlineAchievementsInterfaceGooglePlay.h"
+#include "OnlineLeaderboardInterfaceGooglePlay.h"
+#include "OnlineExternalUIInterfaceGooglePlay.h"
+#include "OnlinePurchaseGooglePlay.h"
+#include "OnlineStoreGooglePlay.h"
 #include "AndroidApplication.h"
 #include "Android/AndroidJNI.h"
 #include "OnlineAsyncTaskManagerGooglePlay.h"
@@ -12,12 +19,14 @@
 #include "Misc/ConfigCacheIni.h"
 #include "Async/TaskGraphInterfaces.h"
 
+THIRD_PARTY_INCLUDES_START
 #include <android_native_app_glue.h>
 
 #include "gpg/android_initialization.h"
 #include "gpg/builder.h"
 #include "gpg/debug.h"
 #include "gpg/android_support.h"
+THIRD_PARTY_INCLUDES_END
 
 using namespace gpg;
 
@@ -43,6 +52,16 @@ IOnlineStorePtr FOnlineSubsystemGooglePlay::GetStoreInterface() const
 	return StoreInterface;
 }
 
+IOnlineStoreV2Ptr FOnlineSubsystemGooglePlay::GetStoreV2Interface() const
+{
+	return StoreV2Interface;
+}
+
+IOnlinePurchasePtr FOnlineSubsystemGooglePlay::GetPurchaseInterface() const
+{
+	return PurchaseInterface;
+}
+
 IOnlineSessionPtr FOnlineSubsystemGooglePlay::GetSessionInterface() const
 {
 	return nullptr;
@@ -113,7 +132,7 @@ static bool WaitingForLogin = false;
 
 bool FOnlineSubsystemGooglePlay::Init() 
 {
-	FPlatformMisc::LowLevelOutputDebugStringf(TEXT("FOnlineSubsystemAndroid::Init"));
+	FPlatformMisc::LowLevelOutputDebugStringf(TEXT("FOnlineSubsystemGooglePlay::Init"));
 	
 	OnlineAsyncTaskThreadRunnable.Reset(new FOnlineAsyncTaskManagerGooglePlay);
 	OnlineAsyncTaskThread.Reset(FRunnableThread::Create(OnlineAsyncTaskThreadRunnable.Get(), *FString::Printf(TEXT("OnlineAsyncTaskThread %s"), *InstanceName.ToString())));
@@ -125,7 +144,18 @@ bool FOnlineSubsystemGooglePlay::Init()
 
 	if (IsInAppPurchasingEnabled())
 	{
-		StoreInterface = MakeShareable(new FOnlineStoreGooglePlay(this));
+		if (IsV2StoreEnabled())
+		{
+			StoreV2Interface = MakeShareable(new FOnlineStoreGooglePlayV2(this));
+			StoreV2Interface->Init();
+			PurchaseInterface = MakeShareable(new FOnlinePurchaseGooglePlay(this));
+			PurchaseInterface->Init();
+		}
+		else
+		{
+			StoreInterface = MakeShareable(new FOnlineStoreGooglePlay(this));
+			StoreInterface->Init();
+		}
 	}
 	
 	extern struct android_app* GNativeAndroidApp;
@@ -158,7 +188,7 @@ bool FOnlineSubsystemGooglePlay::Tick(float DeltaTime)
 
 bool FOnlineSubsystemGooglePlay::Shutdown() 
 {
-	UE_LOG(LogOnline, Log, TEXT("FOnlineSubsystemAndroid::Shutdown()"));
+	UE_LOG(LogOnline, Log, TEXT("FOnlineSubsystemGooglePlay::Shutdown()"));
 
 	FOnlineSubsystemImpl::Shutdown();
 
@@ -173,6 +203,8 @@ bool FOnlineSubsystemGooglePlay::Shutdown()
 
 	// Destruct the interfaces
 	DESTRUCT_INTERFACE(StoreInterface);
+	DESTRUCT_INTERFACE(StoreV2Interface);
+	DESTRUCT_INTERFACE(PurchaseInterface);
 	DESTRUCT_INTERFACE(ExternalUIInterface);
 	DESTRUCT_INTERFACE(AchievementsInterface);
 	DESTRUCT_INTERFACE(LeaderboardsInterface);
@@ -204,16 +236,28 @@ bool FOnlineSubsystemGooglePlay::Exec(UWorld* InWorld, const TCHAR* Cmd, FOutput
 
 bool FOnlineSubsystemGooglePlay::IsEnabled(void)
 {
-	bool bEnabled = true;
-	GConfig->GetBool(TEXT("/Script/AndroidRuntimeSettings.AndroidRuntimeSettings"), TEXT("bEnableGooglePlaySupport"), bEnabled, GEngineIni);
-	return bEnabled;
+	bool bEnableGooglePlaySupport = true;
+	GConfig->GetBool(TEXT("/Script/AndroidRuntimeSettings.AndroidRuntimeSettings"), TEXT("bEnableGooglePlaySupport"), bEnableGooglePlaySupport, GEngineIni);
+
+	FPlatformMisc::LowLevelOutputDebugStringf(TEXT("FOnlineSubsystemGooglePlay::IsEnabled %d"), bEnableGooglePlaySupport);
+	return bEnableGooglePlaySupport;
+}
+
+bool FOnlineSubsystemGooglePlay::IsV2StoreEnabled()
+{
+	bool bUseStoreV2 = false;
+	GConfig->GetBool(TEXT("OnlineSubsystemGooglePlay.Store"), TEXT("bUseStoreV2"), bUseStoreV2, GEngineIni);
+	FPlatformMisc::LowLevelOutputDebugStringf(TEXT("FOnlineSubsystemGooglePlay::IsV2StoreEnabled %d"), bUseStoreV2);
+	return bUseStoreV2;
 }
 
 bool FOnlineSubsystemGooglePlay::IsInAppPurchasingEnabled()
 {
-	bool bEnabledIAP = false;
-	GConfig->GetBool(TEXT("OnlineSubsystemGooglePlay.Store"), TEXT("bSupportsInAppPurchasing"), bEnabledIAP, GEngineIni);
-	return bEnabledIAP;
+	bool bSupportsInAppPurchasing = false;
+	GConfig->GetBool(TEXT("OnlineSubsystemGooglePlay.Store"), TEXT("bSupportsInAppPurchasing"), bSupportsInAppPurchasing, GEngineIni);
+
+	FPlatformMisc::LowLevelOutputDebugStringf(TEXT("FOnlineSubsystemGooglePlay::IsInAppPurchasingEnabled %d"), bSupportsInAppPurchasing);
+	return bSupportsInAppPurchasing;
 }
 
 void FOnlineSubsystemGooglePlay::StartShowLoginUITask(int PlayerId, const FOnLoginUIClosedDelegate& Delegate)
diff --git a/Engine/Plugins/Online/Android/OnlineSubsystemGooglePlay/Source/Public/OnlinePurchaseGooglePlay.h b/Engine/Plugins/Online/Android/OnlineSubsystemGooglePlay/Source/Public/OnlinePurchaseGooglePlay.h
new file mode 100644
--- /dev/null
+++ b/Engine/Plugins/Online/Android/OnlineSubsystemGooglePlay/Source/Public/OnlinePurchaseGooglePlay.h
@@ -0,0 +1,208 @@
+// Copyright 1998-2016 Epic Games, Inc. All Rights Reserved.
+
+#pragma once
+
+#include "OnlinePurchaseInterface.h"
+#include "Serialization/JsonSerializerMacros.h"
+
+enum class EGooglePlayBillingResponseCode : uint8;
+
+/**
+ * Holds in a common format the data that comes out of a Google purchase transaction
+ */
+struct FGoogleTransactionData
+{
+	FGoogleTransactionData(const FString& InOfferId, const FString& InProductToken, const FString& InReceiptData, const FString& InSignature);
+
+	/** @return a string that prints useful debug information about this transaction */
+	FString ToDebugString() const
+	{
+		return FString::Printf(TEXT("OfferId: %s TransactionId: %s ReceiptData: %s%s"),
+			*OfferId,
+			*TransactionIdentifier,
+			*CombinedTransactionData.ToJson(),
+			ErrorStr.IsEmpty() ? TEXT("") : *FString::Printf(TEXT(" Error: %s"), *ErrorStr));
+	}
+
+	/** @return offer id for this transaction */
+	const FString& GetOfferId() const { return OfferId; }
+	/** @return receipt data for this transaction */
+	const FString GetCombinedReceiptData() const { return CombinedTransactionData.ToJson(); }
+	/** @return receipt data for this transaction */
+	const FString& GetReceiptData() const { return CombinedTransactionData.ReceiptData; }
+	/** @return signature for this transaction */
+	const FString& GetSignature() const { return CombinedTransactionData.Signature; }
+	/** @return error string for this transaction, if applicable */
+	const FString& GetErrorStr() const { return ErrorStr; }
+	/** @return the transaction id */
+	const FString& GetTransactionIdentifier() const	{ return TransactionIdentifier; }
+
+private:
+
+	/** Easy access to transmission of data required for backend validation */
+	class FJsonReceiptData :
+		public FJsonSerializable
+	{
+	public:
+
+		FJsonReceiptData() {}
+		FJsonReceiptData(const FString& InReceiptData, const FString& InSignature)
+			: ReceiptData(InReceiptData)
+			, Signature(InSignature)
+		{ }
+
+		/** Opaque store receipt data */
+		FString ReceiptData;
+		/** Signature associated with the transaction */
+		FString Signature;
+
+		// FJsonSerializable
+		BEGIN_JSON_SERIALIZER
+			JSON_SERIALIZE("receiptData", ReceiptData);
+			JSON_SERIALIZE("signature", Signature);
+		END_JSON_SERIALIZER
+	};
+
+	/** GooglePlay offer id */
+	FString OfferId;
+	/** Unique transaction id (purchaseToken) */
+	FString TransactionIdentifier;
+	/** Error on the transaction, if applicable */
+	FString ErrorStr;
+	/** Combined receipt with signature in JSON */
+	FJsonReceiptData CombinedTransactionData;
+};
+
+/**
+ * Info used to cache and track orders in progress.
+ */
+class FOnlinePurchasePendingTransactionGooglePlay
+{
+public:
+	FOnlinePurchasePendingTransactionGooglePlay(
+		const FPurchaseCheckoutRequest& InCheckoutRequest,
+		const FUniqueNetId& InUserId,
+		const EPurchaseTransactionState InPendingTransactionState = EPurchaseTransactionState::NotStarted,
+		const FOnPurchaseCheckoutComplete& InCheckoutCompleteDelegate = FOnPurchaseCheckoutComplete()
+		)
+		: CheckoutRequest(InCheckoutRequest)
+		, UserId(InUserId)
+		, CheckoutCompleteDelegate(InCheckoutCompleteDelegate)
+	{
+		PendingPurchaseInfo.TransactionState = InPendingTransactionState;
+	}
+	
+	/**
+	 * Generate a final receipt for all purchases made in this single transaction
+	 */
+	TSharedRef<FPurchaseReceipt> GenerateReceipt();
+	
+	/** Generate one off receipts for transactions initiated outside the current run of the application */
+	static TSharedRef<FPurchaseReceipt> GenerateReceipt(const FGoogleTransactionData& Transaction);
+	
+	/** Add the single completed transaction to this transaction */
+	bool AddCompletedOffer(EPurchaseTransactionState Result, const FGoogleTransactionData& Transaction);
+
+public:
+	
+	/** Checkout info for the pending order */
+	const FPurchaseCheckoutRequest CheckoutRequest;
+	/** User for the pending order */
+	const FUniqueNetIdString UserId;
+	/** Delegate to call on completion */
+	const FOnPurchaseCheckoutComplete CheckoutCompleteDelegate;
+	/** Tracks the current state of the order */
+	FPurchaseReceipt PendingPurchaseInfo;
+};
+
+/**
+ * Implementation for online purchase via GooglePlay services
+ */
+class FOnlinePurchaseGooglePlay :
+	public IOnlinePurchase,
+	public TSharedFromThis<FOnlinePurchaseGooglePlay, ESPMode::ThreadSafe>
+{
+
+public:
+
+	// IOnlinePurchase
+
+	virtual bool IsAllowedToPurchase(const FUniqueNetId& UserId) override;
+	virtual void Checkout(const FUniqueNetId& UserId, const FPurchaseCheckoutRequest& CheckoutRequest, const FOnPurchaseCheckoutComplete& Delegate) override;
+	virtual void FinalizePurchase(const FUniqueNetId& UserId, const FString& ReceiptId) override;
+	virtual void RedeemCode(const FUniqueNetId& UserId, const FRedeemCodeRequest& RedeemCodeRequest, const FOnPurchaseRedeemCodeComplete& Delegate) override;
+	virtual void QueryReceipts(const FUniqueNetId& UserId, bool bRestoreReceipts, const FOnQueryReceiptsComplete& Delegate) override;
+	virtual void GetReceipts(const FUniqueNetId& UserId, TArray<FPurchaseReceipt>& OutReceipts) const override;
+	
+	// FOnlinePurchaseGooglePlay
+
+	/**
+	 * Constructor
+	 *
+	 * @param InSubsystem GooglePlay subsystem being used
+	 */
+	FOnlinePurchaseGooglePlay(class FOnlineSubsystemGooglePlay* InSubsystem);
+
+	/**
+	 * Constructor
+	 */
+	FOnlinePurchaseGooglePlay();
+
+	/**
+	 * Destructor
+	 */
+	virtual ~FOnlinePurchaseGooglePlay();
+	
+	/** Initialize the interface */
+	void Init();
+
+private:
+	
+	/** Mapping from user id to pending transaction */
+	typedef TMap<FString, TSharedRef<FOnlinePurchasePendingTransactionGooglePlay> > FOnlinePurchasePendingTransactionMap;
+
+	/** Mapping from user id to complete transactions */
+	typedef TMap<FString, TArray< TSharedRef<FPurchaseReceipt> > > FOnlinePurchaseCompleteTransactionsMap;
+	
+	/** Array of transactions completion indirectly (previous run, etc) */
+	typedef TArray< TSharedRef<FPurchaseReceipt> > FOnlineCompletedTransactions;
+	
+private:
+	
+	/** Are receipts being queried */
+	bool bQueryingReceipts;
+	/** Transient delegate to fire when query receipts has completed */
+	FOnQueryReceiptsComplete QueryReceiptsComplete;
+
+	/** Keeps track of pending user transactions */
+	FOnlinePurchasePendingTransactionMap PendingTransactions;
+	
+	/** Cache of completed transactions */
+	FOnlinePurchaseCompleteTransactionsMap CompletedTransactions;
+	
+	/** Cache of purchases completed outside the running instance */
+	FOnlineCompletedTransactions OfflineTransactions;
+	
+	/** Reference to the parent subsystem */
+	FOnlineSubsystemGooglePlay* Subsystem;
+
+	/**
+	 * Delegate fired when a purchase has completed
+	 *
+	 * @param InResponseCode response from the GooglePlay backend
+	 * @param InTransactionData transaction data for the completed purchase
+	 */
+	void OnTransactionCompleteResponse(EGooglePlayBillingResponseCode InResponseCode, const FGoogleTransactionData& InTransactionData);
+	FDelegateHandle ProcessPurchaseResultDelegateHandle;
+
+	/**
+	 * Delegate fired when a query for existing purchase has completed
+	 *
+	 * @param InResponseCode response from the GooglePlay backend
+	 * @param InExistingPurchases known purchases for the user (non consumed or permanent)
+	 */
+	void OnQueryExistingPurchasesComplete(EGooglePlayBillingResponseCode InResponseCode, const TArray<FGoogleTransactionData>& InExistingPurchases);
+	FDelegateHandle QueryExistingPurchasesCompleteDelegateHandle;
+};
+
+typedef TSharedPtr<FOnlinePurchaseGooglePlay, ESPMode::ThreadSafe> FOnlinePurchaseGooglePlayPtr;
\ No newline at end of file
diff --git a/Engine/Plugins/Online/Android/OnlineSubsystemGooglePlay/Source/Public/OnlineSubsystemGooglePlay.h b/Engine/Plugins/Online/Android/OnlineSubsystemGooglePlay/Source/Public/OnlineSubsystemGooglePlay.h
--- a/Engine/Plugins/Online/Android/OnlineSubsystemGooglePlay/Source/Public/OnlineSubsystemGooglePlay.h
+++ b/Engine/Plugins/Online/Android/OnlineSubsystemGooglePlay/Source/Public/OnlineSubsystemGooglePlay.h
@@ -5,17 +5,25 @@
 #include "CoreMinimal.h"
 #include "OnlineSubsystem.h"
 #include "OnlineSubsystemImpl.h"
-#include "OnlineIdentityInterfaceGooglePlay.h"
-#include "OnlineAchievementsInterfaceGooglePlay.h"
-#include "OnlineLeaderboardInterfaceGooglePlay.h"
 #include "OnlineExternalUIInterfaceGooglePlay.h"
 #include "OnlineAsyncTaskManagerGooglePlay.h"
 #include "UniquePtr.h"
 #include "HAL/RunnableThread.h"
 #include "OnlineStoreInterfaceGooglePlay.h"
 
+THIRD_PARTY_INCLUDES_START
 #include "gpg/game_services.h"
 #include "gpg/android_platform_configuration.h"
+THIRD_PARTY_INCLUDES_END
+
+/** Forward declarations of all interface classes */
+typedef TSharedPtr<class FOnlineIdentityGooglePlay,  ESPMode::ThreadSafe> FOnlineIdentityGooglePlayPtr;
+typedef TSharedPtr<class FOnlineStoreGooglePlay, ESPMode::ThreadSafe> FOnlineStoreGooglePlayPtr;
+typedef TSharedPtr<class FOnlineStoreGooglePlayV2, ESPMode::ThreadSafe> FOnlineStoreGooglePlayV2Ptr;
+typedef TSharedPtr<class FOnlinePurchaseGooglePlay, ESPMode::ThreadSafe> FOnlinePurchaseGooglePlayPtr;
+typedef TSharedPtr<class FOnlineLeaderboardsGooglePlay, ESPMode::ThreadSafe> FOnlineLeaderboardsGooglePlayPtr;
+typedef TSharedPtr<class FOnlineAchievementsGooglePlay, ESPMode::ThreadSafe> FOnlineAchievementsGooglePlayPtr;
+typedef TSharedPtr<class FOnlineExternalUIGooglePlay, ESPMode::ThreadSafe> FOnlineExternalUIGooglePlayPtr;
 
 class FOnlineAsyncTaskGooglePlayLogin;
 class FOnlineAsyncTaskGooglePlayShowLoginUI;
@@ -47,8 +55,8 @@ public:
 	virtual IOnlineTitleFilePtr GetTitleFileInterface() const override;
 	virtual IOnlineEntitlementsPtr GetEntitlementsInterface() const override;
 	virtual IOnlineStorePtr GetStoreInterface() const override;
-	virtual IOnlineStoreV2Ptr GetStoreV2Interface() const override { return nullptr; }
-	virtual IOnlinePurchasePtr GetPurchaseInterface() const override { return nullptr; }
+	virtual IOnlineStoreV2Ptr GetStoreV2Interface() const override;
+	virtual IOnlinePurchasePtr GetPurchaseInterface() const override;
 	virtual IOnlineEventsPtr GetEventsInterface() const override { return nullptr; }
 	virtual IOnlineMessagePtr GetMessageInterface() const override { return nullptr; }
 	virtual IOnlineSharingPtr GetSharingInterface() const override { return nullptr; }
@@ -103,11 +111,53 @@ PACKAGE_SCOPE:
 	FOnlineAchievementsGooglePlayPtr GetAchievementsGooglePlay() const { return AchievementsInterface; }
 
 	/**
-	* Is IAP available for use
-	* @return true if IAP should be available, false otherwise
-	*/
+	 * Is IAP available for use
+	 * @return true if IAP should be available, false otherwise
+	 */
 	bool IsInAppPurchasingEnabled();
 
+	/**
+	 * Is Store v2 enabled (disabling legacy store interface)
+	 * @return true if enabled, false otherwise
+	 */
+	bool IsV2StoreEnabled();
+
+	/**
+	 * Delegate fired internally when the Java query for available in app purchases has completed, notifying any GooglePlay OSS listeners
+	 * Not meant for external use
+	 *
+	 * @param InResponseCode response from the GooglePlay backend
+	 * @param ProvidedProductInformation information returned from the backend about the queried offers
+	 */
+	DEFINE_ONLINE_DELEGATE_TWO_PARAM(OnGooglePlayAvailableIAPQueryComplete, EGooglePlayBillingResponseCode /*InResponseCode*/, const TArray<FInAppPurchaseProductInfo>& /*ProvidedProductInformation*/);
+
+	/**
+	 * Delegate fired internally when the Java in app purchase has completed, notifying any GooglePlay OSS listeners
+	 * Not meant for external use
+	 *
+	 * @param InResponseCode response from the GooglePlay backend
+	 * @param InTransactionData transaction information for the purchase
+	 */
+	DEFINE_ONLINE_DELEGATE_TWO_PARAM(OnGooglePlayProcessPurchaseComplete, EGooglePlayBillingResponseCode /*InResponseCode*/, const FGoogleTransactionData& /*InTransactionData*/);
+
+	/**
+	 * Delegate fired internally when the Java query for existing purchases has completed, notifying any GooglePlay OSS listeners
+	 * Not meant for external use
+	 *
+	 * @param InResponseCode response from the GooglePlay backend
+	 * @param InExistingPurchases known purchases for the user (non consumed or permanent)
+	 */
+	DEFINE_ONLINE_DELEGATE_TWO_PARAM(OnGooglePlayQueryExistingPurchasesComplete, EGooglePlayBillingResponseCode /*InResponseCode*/, const TArray<FGoogleTransactionData>& /*InExistingPurchases*/);
+
+	/**
+	 * Delegate fired internally when the Java query for restoring purchases has completed, notifying any GooglePlay OSS listeners
+	 * Not meant for external use
+	 *
+	 * @param InResponseCode response from the GooglePlay backend
+	 * @param InRestoredPurchases restored for the user (non consumed or permanent)
+	 */
+	DEFINE_ONLINE_DELEGATE_TWO_PARAM(OnGooglePlayRestorePurchasesComplete, EGooglePlayBillingResponseCode /*InResponseCode*/, const TArray<FGoogleTransactionData>& /*InRestoredPurchases*/);
+
 	/** Returns true if there are any async login tasks currently being tracked. */
 	bool AreAnyAsyncLoginTasksRunning() const { return CurrentLoginTask != nullptr || CurrentShowLoginUITask != nullptr; }
 
@@ -142,6 +192,12 @@ private:
 
 	FOnlineStoreGooglePlayPtr StoreInterface;
 
+	/** Interface to the online catalog */
+	FOnlineStoreGooglePlayV2Ptr StoreV2Interface;
+
+	/** Interface to the store purchasing */
+	FOnlinePurchaseGooglePlayPtr PurchaseInterface;
+
 	/** Interface to the online leaderboards */
 	FOnlineLeaderboardsGooglePlayPtr LeaderboardsInterface;
 
diff --git a/Engine/Plugins/Online/IOS/OnlineSubsystemIOS/Source/Private/OnlineExternalUIInterfaceIOS.cpp b/Engine/Plugins/Online/IOS/OnlineSubsystemIOS/Source/Private/OnlineExternalUIInterfaceIOS.cpp
--- a/Engine/Plugins/Online/IOS/OnlineSubsystemIOS/Source/Private/OnlineExternalUIInterfaceIOS.cpp
+++ b/Engine/Plugins/Online/IOS/OnlineSubsystemIOS/Source/Private/OnlineExternalUIInterfaceIOS.cpp
@@ -46,7 +46,7 @@ bool FOnlineExternalUIIOS::ShowFriendsUI(int32 LocalUserNum)
 	return false;
 }
 
-bool FOnlineExternalUIIOS::ShowInviteUI(int32 LocalUserNum, FName SessionMame)
+bool FOnlineExternalUIIOS::ShowInviteUI(int32 LocalUserNum, FName SessionName)
 {
 	return false;
 }
diff --git a/Engine/Plugins/Online/IOS/OnlineSubsystemIOS/Source/Private/OnlineExternalUIInterfaceIOS.h b/Engine/Plugins/Online/IOS/OnlineSubsystemIOS/Source/Private/OnlineExternalUIInterfaceIOS.h
--- a/Engine/Plugins/Online/IOS/OnlineSubsystemIOS/Source/Private/OnlineExternalUIInterfaceIOS.h
+++ b/Engine/Plugins/Online/IOS/OnlineSubsystemIOS/Source/Private/OnlineExternalUIInterfaceIOS.h
@@ -18,7 +18,7 @@ public:
 	//~ Begin IOnlineExternalUI Interface
 	virtual bool ShowLoginUI(const int ControllerIndex, bool bShowOnlineOnly, const FOnLoginUIClosedDelegate& Delegate = FOnLoginUIClosedDelegate()) override;
 	virtual bool ShowFriendsUI(int32 LocalUserNum) override;
-	virtual bool ShowInviteUI(int32 LocalUserNum, FName SessionMame = GameSessionName) override;
+	virtual bool ShowInviteUI(int32 LocalUserNum, FName SessionName = GameSessionName) override;
 	virtual bool ShowAchievementsUI(int32 LocalUserNum) override;
 	virtual bool ShowLeaderboardUI(const FString& LeaderboardName) override;
 	virtual bool ShowWebURL(const FString& Url, const FShowWebUrlParams& ShowParams, const FOnShowWebUrlClosedDelegate& Delegate = FOnShowWebUrlClosedDelegate()) override;
diff --git a/Engine/Plugins/Online/IOS/OnlineSubsystemIOS/Source/Private/OnlinePurchaseIOS.cpp b/Engine/Plugins/Online/IOS/OnlineSubsystemIOS/Source/Private/OnlinePurchaseIOS.cpp
new file mode 100644
--- /dev/null
+++ b/Engine/Plugins/Online/IOS/OnlineSubsystemIOS/Source/Private/OnlinePurchaseIOS.cpp
@@ -0,0 +1,501 @@
+// Copyright 1998-2016 Epic Games, Inc. All Rights Reserved.
+
+#include "OnlineSubsystemIOSPrivatePCH.h"
+#include "OnlinePurchaseIOS.h"
+#include "OnlineError.h"
+#import "OnlineStoreKitHelper.h"
+
+namespace OSSConsoleVariables
+{
+	// CVars
+	TAutoConsoleVariable<int32> CVarSimluateAskToBuy(
+															 TEXT("OSS.AskToBuy"),
+															 0,
+															 TEXT("Simulate ask to buy in iOS\n")
+															 TEXT("1 enable, 0 disable"),
+															 ECVF_Default);
+}
+
+#define LOCTEXT_NAMESPACE "OnlineSubsystemIOS"
+#define IOSUSER TEXT("IOSUser")
+
+FOnlinePurchaseIOS::FOnlinePurchaseIOS(FOnlineSubsystemIOS* InSubsystem)
+	: StoreHelper(nullptr)
+	, bRestoringTransactions(false)
+	, Subsystem(InSubsystem)
+{
+	UE_LOG(LogOnline, Verbose, TEXT( "FOnlinePurchaseIOS::FOnlinePurchaseIOS" ));
+}
+
+FOnlinePurchaseIOS::FOnlinePurchaseIOS()
+{
+	UE_LOG(LogOnline, Verbose, TEXT( "FOnlinePurchaseIOS::FOnlinePurchaseIOS" ));
+}
+
+FOnlinePurchaseIOS::~FOnlinePurchaseIOS()
+{
+}
+
+void FOnlinePurchaseIOS::InitStoreKit(FStoreKitHelperV2* InStoreKit)
+{
+	StoreHelper = InStoreKit;
+
+	FOnProductsRequestResponseDelegate OnProductsRequestResponseDelegate;
+    OnProductsRequestResponseDelegate.BindRaw(this, &FOnlinePurchaseIOS::OnProductPurchaseRequestResponse);
+	[StoreHelper AddOnProductRequestResponse: OnProductsRequestResponseDelegate];
+	
+	FOnTransactionCompleteIOSDelegate OnTransactionCompleteResponseDelegate = FOnTransactionCompleteIOSDelegate::CreateRaw(this, &FOnlinePurchaseIOS::OnTransactionCompleteResponse);
+	[StoreHelper AddOnTransactionComplete: OnTransactionCompleteResponseDelegate];
+	
+	FOnTransactionRestoredIOSDelegate OnTransactionRestoredDelegate = FOnTransactionRestoredIOSDelegate::CreateRaw(this, &FOnlinePurchaseIOS::OnTransactionRestored);
+	[StoreHelper AddOnTransactionRestored: OnTransactionRestoredDelegate];
+	
+	FOnRestoreTransactionsCompleteIOSDelegate OnRestoreTransactionsCompleteDelegate = FOnRestoreTransactionsCompleteIOSDelegate::CreateRaw(this, &FOnlinePurchaseIOS::OnRestoreTransactionsComplete);
+	[StoreHelper AddOnRestoreTransactionsComplete: OnRestoreTransactionsCompleteDelegate];
+	
+	FOnTransactionProgressDelegate OnTransactionPurchasingDelegate = FOnTransactionProgressDelegate::CreateRaw(this, &FOnlinePurchaseIOS::OnTransactionInProgress);
+	[StoreHelper AddOnPurchaseInProgress: OnTransactionPurchasingDelegate];
+	
+	FOnTransactionProgressDelegate OnTransactionDeferredDelegate = FOnTransactionProgressDelegate::CreateRaw(this, &FOnlinePurchaseIOS::OnTransactionDeferred);
+	[StoreHelper AddOnTransactionDeferred: OnTransactionDeferredDelegate];
+}
+
+bool FOnlinePurchaseIOS::IsAllowedToPurchase(const FUniqueNetId& UserId)
+{
+	UE_LOG(LogOnline, Verbose, TEXT("FOnlinePurchaseIOS::IsAllowedToPurchase"));
+	bool bCanMakePurchases = [SKPaymentQueue canMakePayments];
+	return bCanMakePurchases;
+}
+
+void FOnlinePurchaseIOS::Checkout(const FUniqueNetId& UserId, const FPurchaseCheckoutRequest& CheckoutRequest, const FOnPurchaseCheckoutComplete& Delegate)
+{
+	bool bStarted = false;
+	FText ErrorMessage;
+
+	TSharedRef<FOnlinePurchasePendingTransactionIOS> RequestedTransaction = MakeShareable(new FOnlinePurchasePendingTransactionIOS(CheckoutRequest, UserId, EPurchaseTransactionState::NotStarted, Delegate));
+
+	if (IsAllowedToPurchase(UserId))
+	{
+		const FString UserIdStr = IOSUSER;
+		const TSharedRef<FOnlinePurchasePendingTransactionIOS>* UserPendingTransaction = PendingTransactions.Find(UserIdStr);
+		if (UserPendingTransaction == nullptr)
+		{
+			FOnlineStoreIOSPtr StoreInterface = StaticCastSharedPtr<FOnlineStoreIOS>(Subsystem->GetStoreV2Interface());
+			if (StoreInterface.IsValid())
+			{
+				TSharedRef<FOnlinePurchasePendingTransactionIOS> PendingTransaction = PendingTransactions.Add(UserIdStr, RequestedTransaction);
+				PendingTransaction->StartProcessing();
+				
+				// autoreleased NSMutableArray to hold products
+				int32 NumOffers = CheckoutRequest.PurchaseOffers.Num();
+
+				NSMutableArray<SKProduct*>* ProductSet = [NSMutableArray arrayWithCapacity: NumOffers];
+				for (int32 OfferIdx = 0; OfferIdx < NumOffers; OfferIdx++)
+				{
+					const FPurchaseCheckoutRequest::FPurchaseOfferEntry& Offer = CheckoutRequest.PurchaseOffers[OfferIdx];
+
+					SKProduct* Product = StoreInterface->GetSKProductByOfferId(Offer.OfferId);
+					if (Product)
+					{
+						[ProductSet addObject:Product];
+					}
+				}
+				
+				bool bAskToBuy = false;
+#if !UE_BUILD_SHIPPING
+				bAskToBuy = OSSConsoleVariables::CVarSimluateAskToBuy.GetValueOnGameThread() == 1;
+#endif
+				
+				if ([ProductSet count] > 0)
+				{
+					dispatch_async(dispatch_get_main_queue(), ^
+					{
+					    // Purchase the product through the StoreKit framework
+						[StoreHelper makePurchase:ProductSet SimulateAskToBuy:bAskToBuy];
+					});
+					bStarted = true;
+				}
+				else
+				{
+					ErrorMessage = NSLOCTEXT("IOSPurchase", "ErrorNoOffersSpecified", "Failed to checkout, no offers given.");
+					RequestedTransaction->PendingPurchaseInfo.TransactionState = EPurchaseTransactionState::Failed;
+				}
+			}
+		}
+		else
+		{
+			ErrorMessage = NSLOCTEXT("IOSPurchase", "ErrorTransactionInProgress", "Failed to checkout, user has in progress transaction.");
+			RequestedTransaction->PendingPurchaseInfo.TransactionState = EPurchaseTransactionState::Failed;
+		}
+	}
+	else
+	{
+		ErrorMessage = NSLOCTEXT("IOSPurchase", "ErrorPurchaseNotAllowed", "Failed to checkout, user not allowed to purchase.");
+		RequestedTransaction->PendingPurchaseInfo.TransactionState = EPurchaseTransactionState::Failed;
+	}
+
+	if (!bStarted)
+	{
+		TSharedRef<FPurchaseReceipt> FailReceipt = RequestedTransaction->GenerateReceipt();
+		
+		Subsystem->ExecuteNextTick([ErrorMessage, FailReceipt, Delegate]()
+		{
+			FOnlineError Error(ErrorMessage);
+			Delegate.ExecuteIfBound(Error, FailReceipt);
+		});
+	}
+}
+
+void FOnlinePurchaseIOS::FinalizePurchase(const FUniqueNetId& UserId, const FString& ReceiptId)
+{
+	UE_LOG(LogOnline, Verbose, TEXT("FOnlinePurchaseIOS::FinalizePurchase %s %s"), *UserId.ToString(), *ReceiptId);
+
+	const FString ReceiptIdCopy(ReceiptId);
+	dispatch_async(dispatch_get_main_queue(), ^
+	{
+		// Purchase the product through the StoreKit framework
+		[StoreHelper finalizeTransaction:ReceiptIdCopy];
+	});
+}
+
+void FOnlinePurchaseIOS::RedeemCode(const FUniqueNetId& UserId, const FRedeemCodeRequest& RedeemCodeRequest, const FOnPurchaseRedeemCodeComplete& Delegate)
+{
+	FOnlineError Result;
+	Delegate.ExecuteIfBound(Result, MakeShareable(new FPurchaseReceipt()));
+}
+
+void FOnlinePurchaseIOS::QueryReceipts(const FUniqueNetId& UserId, bool bRestoreReceipts, const FOnQueryReceiptsComplete& Delegate)
+{
+	bool bSuccess = true;
+	bool bTriggerDelegate = true;
+	if (bRestoreReceipts)
+	{
+		if (!bRestoringTransactions)
+		{
+			// Restore purchases, adding them to the offline receipts array for later redemption
+			if (StoreHelper)
+			{
+				QueryReceiptsComplete = Delegate;
+				bTriggerDelegate = false;
+				bRestoringTransactions = true;
+				dispatch_async(dispatch_get_main_queue(), ^
+				{
+					[StoreHelper restorePurchases];
+				});
+			}
+			else
+			{
+				bSuccess = false;
+			}
+		}
+		else
+		{
+			UE_LOG(LogOnline, Verbose, TEXT("FOnlinePurchaseIOS::QueryReceipts already restoring transactions"));
+			bSuccess = false;
+		}
+	}
+	
+	if (bTriggerDelegate)
+	{
+		// Query receipts comes dynamically from the StoreKit observer
+		Subsystem->ExecuteNextTick([Delegate, bSuccess]() {
+			FOnlineError Result(bSuccess);
+			Delegate.ExecuteIfBound(Result);
+		});
+	}
+}
+
+void FOnlinePurchaseIOS::GetReceipts(const FUniqueNetId& UserId, TArray<FPurchaseReceipt>& OutReceipts) const
+{
+	OutReceipts.Empty();
+
+	// Add the cached list of user purchases
+	const TArray<TSharedRef<FPurchaseReceipt>>* UserCompletedTransactions = CompletedTransactions.Find(UserId.ToString());
+	if (UserCompletedTransactions != nullptr)
+	{
+		for (int32 Idx = 0; Idx < UserCompletedTransactions->Num(); Idx++)
+		{
+			OutReceipts.Add(*(*UserCompletedTransactions)[Idx]);
+		}
+	}
+	
+	// Add purchases completed while "offline"
+	for (int32 Idx = 0; Idx < OfflineTransactions.Num(); Idx++)
+	{
+		OutReceipts.Add(*OfflineTransactions[Idx]);
+	}
+}
+
+void FOnlinePurchaseIOS::OnProductPurchaseRequestResponse(SKProductsResponse* Response, const FOnQueryOnlineStoreOffersComplete& CompletionDelegate)
+{
+	UE_LOG(LogOnline, Verbose, TEXT("FOnlinePurchaseIOS::OnProductPurchaseRequestResponse"));
+}
+
+void FOnlinePurchaseIOS::OnTransactionCompleteResponse(EPurchaseTransactionState Result, const FStoreKitTransactionData& TransactionData)
+{
+	UE_LOG(LogOnline, Verbose, TEXT("FOnlinePurchaseIOS::OnTransactionCompleteResponse %d %s"), (int32)Result, *TransactionData.ToDebugString());
+	
+	FString UserIdStr = IOSUSER;
+	const TSharedRef<FOnlinePurchasePendingTransactionIOS>* UserPendingTransactionPtr = PendingTransactions.Find(UserIdStr);
+	if(UserPendingTransactionPtr != nullptr)
+	{
+		const TSharedRef<FOnlinePurchasePendingTransactionIOS> UserPendingTransaction = *UserPendingTransactionPtr;
+		
+		UserPendingTransaction->AddCompletedOffer(Result, TransactionData);
+		if (UserPendingTransaction->AreAllOffersComplete())
+		{
+			FOnlineError FinalResult;
+			EPurchaseTransactionState FinalState = UserPendingTransaction->GetFinalTransactionState();
+			
+			UserPendingTransaction->PendingPurchaseInfo.TransactionState = FinalState;
+			// UserPendingTransaction->PendingPurchaseInfo.TransactionId; purposefully blank
+			
+			const FString& ErrorStr = TransactionData.GetErrorStr();
+			switch (FinalState)
+			{
+				case EPurchaseTransactionState::Failed:
+					FinalResult.SetFromErrorCode(TEXT("com.epicgames.purchase.failure"));
+					FinalResult.ErrorMessage = !ErrorStr.IsEmpty() ? FText::FromString(ErrorStr) : LOCTEXT("IOSTransactionFailed", "TransactionFailed");
+					break;
+				case EPurchaseTransactionState::Canceled:
+					FinalResult.SetFromErrorCode(TEXT("com.epicgames.catalog_helper.user_cancelled"));
+					FinalResult.ErrorMessage = !ErrorStr.IsEmpty() ? FText::FromString(ErrorStr) : LOCTEXT("IOSTransactionCancel", "TransactionCanceled");
+					break;
+				case EPurchaseTransactionState::Purchased:
+					FinalResult.bSucceeded = true;
+					break;
+				default:
+					UE_LOG(LogOnline, Warning, TEXT("Unexpected state after purchase %d"), FinalState);
+					FinalResult.SetFromErrorCode(TEXT("com.epicgames.purchase.unexpected_state"));
+					FinalResult.ErrorMessage = !ErrorStr.IsEmpty() ? FText::FromString(ErrorStr) : LOCTEXT("UnexpectedState", "Unexpected purchase result");
+					UserPendingTransaction->PendingPurchaseInfo.TransactionState = EPurchaseTransactionState::Failed;
+					break;
+			}
+			
+			TSharedRef<FPurchaseReceipt> FinalReceipt = UserPendingTransaction->GenerateReceipt();
+			
+			TArray< TSharedRef<FPurchaseReceipt> >& UserCompletedTransactions = CompletedTransactions.FindOrAdd(UserIdStr);
+			
+			PendingTransactions.Remove(UserIdStr);
+			UserCompletedTransactions.Add(FinalReceipt);
+			
+			UserPendingTransaction->CheckoutCompleteDelegate.ExecuteIfBound(FinalResult, FinalReceipt);
+		}
+	}
+	else
+	{
+		// Transactions that come in during login or other non explicit purchase moments are added to a receipts list for later redemption
+		UE_LOG(LogOnline, Log, TEXT("Pending transaction completed offline"));
+		
+		if (Result == EPurchaseTransactionState::Restored || Result == EPurchaseTransactionState::Purchased)
+		{
+			TSharedRef<FPurchaseReceipt> OfflineReceipt = FOnlinePurchasePendingTransactionIOS::GenerateReceipt(Result, TransactionData);
+			OfflineTransactions.Add(OfflineReceipt);
+		}
+	}
+}
+
+void FOnlinePurchaseIOS::OnTransactionRestored(const FStoreKitTransactionData& TransactionData)
+{
+	UE_LOG(LogOnline, Verbose, TEXT("FOnlinePurchaseIOS::OnTransactionRestored %s"), *TransactionData.ToDebugString());
+
+	// Single item restored amongst a group of items
+	TSharedRef<FPurchaseReceipt> OfflineReceipt = FOnlinePurchasePendingTransactionIOS::GenerateReceipt(EPurchaseTransactionState::Restored, TransactionData);
+	
+#if 0
+	bool bFound = false;
+	for (TSharedRef<FPurchaseReceipt>& OtherReceipt : OfflineTransactions)
+	{
+		// If redundant, replace the entry
+		if (OtherReceipt->TransactionId == OfflineReceipt->TransactionId)
+		{
+			*OtherReceipt = *OfflineReceipt;
+			bFound = true;
+			break;
+		}
+	}
+	
+	if (!bFound)
+#endif
+	{
+		OfflineTransactions.Add(OfflineReceipt);
+	}
+}
+
+void FOnlinePurchaseIOS::OnRestoreTransactionsComplete(EPurchaseTransactionState Result)
+{
+	UE_LOG(LogOnline, Verbose, TEXT("FOnlinePurchaseIOS::OnRestoreTransactionsComplete %d"), (int32)Result);
+	
+	// Full restore is complete
+	bRestoringTransactions = false;
+	bool bSuccess = (Result == EPurchaseTransactionState::Restored) || (Result == EPurchaseTransactionState::Purchased);
+	Subsystem->ExecuteNextTick([this, bSuccess]() {
+		FOnlineError FinalResult(bSuccess);
+		QueryReceiptsComplete.ExecuteIfBound(FinalResult);
+		QueryReceiptsComplete.Unbind();
+	});
+}
+
+void FOnlinePurchaseIOS::OnTransactionInProgress(const FStoreKitTransactionData& TransactionData)
+{
+	UE_LOG(LogOnline, Verbose, TEXT("FOnlinePurchaseIOS::OnTransactionInProgress %s"), *TransactionData.ToDebugString());
+}
+
+void FOnlinePurchaseIOS::OnTransactionDeferred(const FStoreKitTransactionData& TransactionData)
+{
+	UE_LOG(LogOnline, Verbose, TEXT("FOnlinePurchaseIOS::OnTransactionDeferred %s"), *TransactionData.ToDebugString());
+	
+	FString UserIdStr = IOSUSER;
+	const TSharedRef<FOnlinePurchasePendingTransactionIOS>* UserPendingTransactionPtr = PendingTransactions.Find(UserIdStr);
+	if (UserPendingTransactionPtr != nullptr)
+	{
+		const TSharedRef<FOnlinePurchasePendingTransactionIOS> UserPendingTransaction = *UserPendingTransactionPtr;
+		
+		const FString& ErrorStr = TransactionData.GetErrorStr();
+		
+		FOnlineError FinalResult;
+		FinalResult.SetFromErrorCode(TEXT("com.epicgames.purchase.deferred"));
+		FinalResult.ErrorMessage = !ErrorStr.IsEmpty() ? FText::FromString(ErrorStr) : LOCTEXT("IOSTransactionDeferred", "Transaction awaiting approval.");
+		
+		TSharedRef<FPurchaseReceipt> DeferredReceipt = FOnlinePurchasePendingTransactionIOS::GenerateReceipt(EPurchaseTransactionState::Deferred, TransactionData);
+		
+		// Clear out the deferred transaction, it should get picked up in "offline" receipts
+		PendingTransactions.Remove(UserIdStr);
+		UserPendingTransaction->CheckoutCompleteDelegate.ExecuteIfBound(FinalResult, DeferredReceipt);
+	}
+	else
+	{
+		UE_LOG(LogOnline, Log, TEXT("Offline deferred transaction"));
+	}
+}
+
+TSharedRef<FPurchaseReceipt> FOnlinePurchasePendingTransactionIOS::GenerateReceipt()
+{
+	TSharedRef<FPurchaseReceipt> Receipt = MakeShareable(new FPurchaseReceipt());
+	
+	Receipt->TransactionState = PendingPurchaseInfo.TransactionState;
+	Receipt->TransactionId = PendingPurchaseInfo.TransactionId;
+	
+	if(PendingPurchaseInfo.TransactionState == EPurchaseTransactionState::Purchased ||
+	   PendingPurchaseInfo.TransactionState == EPurchaseTransactionState::Restored)
+	{
+		Receipt->ReceiptOffers = PendingPurchaseInfo.ReceiptOffers;
+	}
+	else
+	{
+		// Add the requested offers to the receipt in the event of an incomplete purchase.
+		for(const auto& RequestedOffer : CheckoutRequest.PurchaseOffers)
+		{
+			Receipt->AddReceiptOffer(RequestedOffer.OfferNamespace, RequestedOffer.OfferId, RequestedOffer.Quantity);
+		}
+	}
+	
+	return Receipt;
+}
+
+TSharedRef<FPurchaseReceipt> FOnlinePurchasePendingTransactionIOS::GenerateReceipt(EPurchaseTransactionState Result, const FStoreKitTransactionData& Transaction)
+{
+	TSharedRef<FPurchaseReceipt> Receipt = MakeShareable(new FPurchaseReceipt());
+	
+	Receipt->TransactionState = Result;
+	Receipt->TransactionId = Transaction.GetTransactionIdentifier();
+	
+	if (Result == EPurchaseTransactionState::Purchased ||
+	    Result == EPurchaseTransactionState::Restored)
+	{
+		FPurchaseReceipt::FReceiptOfferEntry ReceiptEntry(TEXT(""), Transaction.GetOfferId(), 1);
+		
+		int32 Idx = ReceiptEntry.LineItems.AddZeroed();
+		
+		FPurchaseReceipt::FLineItemInfo& LineItem = ReceiptEntry.LineItems[Idx];
+		
+		LineItem.ItemName = Transaction.GetOfferId();
+		LineItem.UniqueId = Transaction.GetTransactionIdentifier();;
+		LineItem.ValidationInfo = Transaction.GetReceiptData();
+
+		Receipt->AddReceiptOffer(ReceiptEntry);
+	}
+	
+	return Receipt;
+}
+
+void FOnlinePurchasePendingTransactionIOS::StartProcessing()
+{
+	PendingPurchaseInfo.TransactionState = EPurchaseTransactionState::Processing;
+	for (int32 OfferIdx = 0; OfferIdx < CheckoutRequest.PurchaseOffers.Num(); ++OfferIdx)
+	{
+		OfferPurchaseStates[OfferIdx] = EPurchaseTransactionState::Processing;
+	}
+}
+
+bool FOnlinePurchasePendingTransactionIOS::AddCompletedOffer(EPurchaseTransactionState Result, const FStoreKitTransactionData& Transaction)
+{
+	for (int32 OfferIdx = 0; OfferIdx < CheckoutRequest.PurchaseOffers.Num(); ++OfferIdx)
+	{
+		const FPurchaseCheckoutRequest::FPurchaseOfferEntry& Offer = CheckoutRequest.PurchaseOffers[OfferIdx];
+		if (Transaction.GetOfferId() == Offer.OfferId)
+		{
+			OfferPurchaseStates[OfferIdx] = Result;
+			FPurchaseReceipt::FReceiptOfferEntry Receipt(TEXT(""), Transaction.GetOfferId(), 1);
+
+			int32 Idx = Receipt.LineItems.AddZeroed();
+			
+			FPurchaseReceipt::FLineItemInfo& LineItem = Receipt.LineItems[Idx];
+			
+			LineItem.ItemName = Transaction.GetOfferId();
+			LineItem.UniqueId = Transaction.GetTransactionIdentifier();
+			LineItem.ValidationInfo = Transaction.GetReceiptData();
+			
+			PendingPurchaseInfo.AddReceiptOffer(Receipt);
+			return true;
+		}
+	}
+	
+	return false;
+}
+
+bool FOnlinePurchasePendingTransactionIOS::AreAllOffersComplete() const
+{
+	for (const EPurchaseTransactionState State : OfferPurchaseStates)
+	{
+		if (State == EPurchaseTransactionState::NotStarted ||
+			State == EPurchaseTransactionState::Processing)
+		{
+			return false;
+		}
+	}
+	
+	return true;
+}
+
+EPurchaseTransactionState FOnlinePurchasePendingTransactionIOS::GetFinalTransactionState() const
+{
+	bool bAnyFailures = false;
+	bool bAnyCancels = false;
+	for (const EPurchaseTransactionState State : OfferPurchaseStates)
+	{
+		if (State == EPurchaseTransactionState::NotStarted ||
+			State == EPurchaseTransactionState::Processing ||
+			State == EPurchaseTransactionState::Failed)
+		{
+			bAnyFailures = true;
+		}
+		else if (State == EPurchaseTransactionState::Canceled)
+		{
+			bAnyCancels = true;
+		}
+	}
+	
+	if (bAnyFailures)
+	{
+		return EPurchaseTransactionState::Failed;
+	}
+	else if (bAnyCancels)
+	{
+		return EPurchaseTransactionState::Canceled;
+	}
+	
+	return EPurchaseTransactionState::Purchased;
+}
+
+#undef LOCTEXT_NAMESPACE
diff --git a/Engine/Plugins/Online/IOS/OnlineSubsystemIOS/Source/Private/OnlineSessionInterfaceIOS.cpp b/Engine/Plugins/Online/IOS/OnlineSubsystemIOS/Source/Private/OnlineSessionInterfaceIOS.cpp
--- a/Engine/Plugins/Online/IOS/OnlineSubsystemIOS/Source/Private/OnlineSessionInterfaceIOS.cpp
+++ b/Engine/Plugins/Online/IOS/OnlineSubsystemIOS/Source/Private/OnlineSessionInterfaceIOS.cpp
@@ -667,7 +667,7 @@ bool FOnlineSessionIOS::GetResolvedConnectString(FName SessionName, FString& Con
 }
 
 
-bool FOnlineSessionIOS::GetResolvedConnectString(const class FOnlineSessionSearchResult& SearchResult, FName PortType, FString& ConnectInfo)
+bool FOnlineSessionIOS::GetResolvedConnectString(const FOnlineSessionSearchResult& SearchResult, FName PortType, FString& ConnectInfo)
 {
 	return false;
 }
diff --git a/Engine/Plugins/Online/IOS/OnlineSubsystemIOS/Source/Private/OnlineStoreIOS.cpp b/Engine/Plugins/Online/IOS/OnlineSubsystemIOS/Source/Private/OnlineStoreIOS.cpp
new file mode 100644
--- /dev/null
+++ b/Engine/Plugins/Online/IOS/OnlineSubsystemIOS/Source/Private/OnlineStoreIOS.cpp
@@ -0,0 +1,240 @@
+// Copyright 1998-2016 Epic Games, Inc. All Rights Reserved.
+
+#include "OnlineSubsystemIOSPrivatePCH.h"
+#include "OnlineStoreIOS.h"
+#include "Internationalization.h"
+#include "Internationalization/Culture.h"
+#include "FastDecimalFormat.h"
+#import "OnlineStoreKitHelper.h"
+
+FOnlineStoreIOS::FOnlineStoreIOS(FOnlineSubsystemIOS* InSubsystem)
+	: StoreHelper(nullptr)
+	, bIsQueryInFlight(false)
+	, Subsystem(InSubsystem)
+{
+	UE_LOG(LogOnline, Verbose, TEXT( "FOnlineStoreIOS::FOnlineStoreIOS" ));
+}
+
+FOnlineStoreIOS::FOnlineStoreIOS()
+	: bIsQueryInFlight(false)
+{
+	UE_LOG(LogOnline, Verbose, TEXT( "FOnlineStoreIOS::FOnlineStoreIOS" ));
+}
+
+void FOnlineStoreIOS::InitStoreKit(FStoreKitHelperV2* InStoreKit)
+{
+	StoreHelper = InStoreKit;
+	
+	FOnProductsRequestResponseDelegate OnProductsRequestResponseDelegate;
+	OnProductsRequestResponseDelegate.BindRaw(this, &FOnlineStoreIOS::OnProductPurchaseRequestResponse);
+	[StoreHelper AddOnProductRequestResponse: OnProductsRequestResponseDelegate];
+
+	// FOnTransactionCompleteIOSDelegate OnTransactionCompleteResponseDelegate = FOnTransactionCompleteIOSDelegate::CreateRaw(this, &FOnlineStoreIOS::OnTransactionCompleteResponse);
+	// [StoreHelper AddOnTransactionComplete: OnTransactionCompleteResponseDelegate];
+	
+	// FOnTransactionRestoredIOSDelegate OnTransactionRestoredDelegate = FOnTransactionRestoredIOSDelegate::CreateRaw(this, &FOnlineStoreIOS::OnTransactionRestored);
+	// [StoreHelper AddOnTransactionRestored: OnTransactionRestoredDelegate];
+	
+	// FOnRestoreTransactionsCompleteIOSDelegate OnRestoreTransactionsCompleteDelegate = FOnRestoreTransactionsCompleteIOSDelegate::CreateRaw(this, &FOnlineStoreIOS::OnRestoreTransactionsComplete);
+	// [StoreHelper AddOnRestoreTransactionsComplete: OnRestoreTransactionsCompleteDelegate];
+}
+
+FOnlineStoreIOS::~FOnlineStoreIOS()
+{
+	[StoreHelper release];
+}
+
+void FOnlineStoreIOS::QueryCategories(const FUniqueNetId& UserId, const FOnQueryOnlineStoreCategoriesComplete& Delegate)
+{
+	Delegate.ExecuteIfBound(false, TEXT("No CatalogService"));
+}
+
+void FOnlineStoreIOS::GetCategories(TArray<FOnlineStoreCategory>& OutCategories) const
+{
+	OutCategories.Empty();
+}
+
+void FOnlineStoreIOS::QueryOffersByFilter(const FUniqueNetId& UserId, const FOnlineStoreFilter& Filter, const FOnQueryOnlineStoreOffersComplete& Delegate)
+{
+	Delegate.ExecuteIfBound(false, TArray<FUniqueOfferId>(), TEXT("No CatalogService"));
+}
+
+void FOnlineStoreIOS::QueryOffersById(const FUniqueNetId& UserId, const TArray<FUniqueOfferId>& OfferIds, const FOnQueryOnlineStoreOffersComplete& Delegate)
+{
+	UE_LOG(LogOnline, Verbose, TEXT("FOnlineStoreIOS::QueryOffersById"));
+	
+	if (bIsQueryInFlight)
+	{
+		Delegate.ExecuteIfBound(false, OfferIds, TEXT("Request already in flight"));
+	}
+	else if (OfferIds.Num() == 0)
+	{
+		Delegate.ExecuteIfBound(false, OfferIds, TEXT("No offers to query for"));
+	}
+	else
+	{
+		// autoreleased NSSet to hold IDs
+		NSMutableSet* ProductSet = [NSMutableSet setWithCapacity:OfferIds.Num()];
+		for (int32 OfferIdx = 0; OfferIdx < OfferIds.Num(); OfferIdx++)
+		{
+			NSString* ID = [NSString stringWithFString:OfferIds[OfferIdx]];
+			// convert to NSString for the set objects
+			[ProductSet addObject:ID];
+		}
+		
+		FOnQueryOnlineStoreOffersComplete DelegateCopy = Delegate;
+		dispatch_async(dispatch_get_main_queue(), ^
+        {
+			[StoreHelper requestProductData:ProductSet WithDelegate:DelegateCopy];
+		});
+		
+		bIsQueryInFlight = true;
+	}
+}
+
+/**
+ * Convert an Apple SKProduct reference into a engine FOnlineStoreOffer
+ * (Apple has only Title/Description converted to Title/(short)Description)
+ *
+ * @param Product information about an Apple offer from iTunesConnect
+ *
+ * @return FOnlineStoreOffer with proper parameters filled in
+ */
+FOnlineStoreOffer ConvertProductToStoreOffer(SKProduct* Product)
+{
+	FOnlineStoreOffer NewProductInfo;
+	
+	NSNumberFormatter *numberFormatter = [[NSNumberFormatter alloc] init];
+	[numberFormatter setFormatterBehavior : NSNumberFormatterBehavior10_4];
+	[numberFormatter setNumberStyle : NSNumberFormatterCurrencyStyle];
+	[numberFormatter setLocale : Product.priceLocale];
+	
+	NewProductInfo.OfferId = [Product productIdentifier];
+	
+	NewProductInfo.Title = FText::FromString([Product localizedTitle]);
+	NewProductInfo.Description = FText::FromString([Product localizedDescription]);
+	//NewProductInfo.LongDescription = FText::FromString([Product localizedDescription]);
+	NewProductInfo.PriceText = FText::FromString([numberFormatter stringFromNumber : Product.price]);
+	NewProductInfo.CurrencyCode = [Product.priceLocale objectForKey : NSLocaleCurrencyCode];
+	
+	// Convert the backend stated price into its base units
+	FInternationalization& I18N = FInternationalization::Get();
+	const FCulture& Culture = *I18N.GetCurrentCulture();
+	
+	const FDecimalNumberFormattingRules& FormattingRules = Culture.GetCurrencyFormattingRules(NewProductInfo.CurrencyCode);
+	const FNumberFormattingOptions& FormattingOptions = FormattingRules.CultureDefaultFormattingOptions;
+	double Val = static_cast<double>([Product.price doubleValue]) * static_cast<double>(FMath::Pow(10.0f, FormattingOptions.MaximumFractionalDigits));
+	NewProductInfo.NumericPrice = FMath::TruncToInt(Val + 0.5);
+	
+	// iOS doesn't support these fields, set to min and max defaults
+	NewProductInfo.ReleaseDate = FDateTime::MinValue();
+	NewProductInfo.ExpirationDate = FDateTime::MaxValue();
+	
+	return NewProductInfo;
+}
+
+void FOnlineStoreIOS::AddOffer(const FOnlineStoreOfferIOS& NewOffer)
+{
+	if (NewOffer.IsValid())
+	{
+		FOnlineStoreOfferIOS* Existing = CachedOffers.Find(NewOffer.Offer->OfferId);
+		if (Existing != nullptr)
+		{
+			// Replace existing offer
+			*Existing = NewOffer;
+		}
+		else
+		{
+			CachedOffers.Add(NewOffer.Offer->OfferId, NewOffer);
+		}
+	}
+}
+
+void FOnlineStoreIOS::OnProductPurchaseRequestResponse(SKProductsResponse* Response, const FOnQueryOnlineStoreOffersComplete& CompletionDelegate)
+{
+	if(bIsQueryInFlight)
+	{
+		bool bWasSuccessful = [Response.products count] > 0;
+		int32 NumInvalidProducts = [Response.invalidProductIdentifiers count];
+		if ([Response.products count] == 0 && NumInvalidProducts == 0)
+		{
+			UE_LOG(LogOnline, Warning, TEXT("Wrong number of products [%d] in the response when trying to make a single purchase"), [Response.products count]);
+		}
+		
+		TArray<FUniqueOfferId> OfferIds;
+		for (SKProduct* Product in Response.products)
+		{
+			FOnlineStoreOfferIOS NewProductOffer(Product, ConvertProductToStoreOffer(Product));
+			
+			AddOffer(NewProductOffer);
+			OfferIds.Add(NewProductOffer.Offer->OfferId);
+			
+			UE_LOG(LogOnline, Log, TEXT("Product Identifier: %s, Name: %s, Desc: %s, Long Desc: %s, Price: %s IntPrice: %d"),
+				   *NewProductOffer.Offer->OfferId,
+				   *NewProductOffer.Offer->Title.ToString(),
+				   *NewProductOffer.Offer->Description.ToString(),
+				   *NewProductOffer.Offer->LongDescription.ToString(),
+				   *NewProductOffer.Offer->PriceText.ToString(),
+				   NewProductOffer.Offer->NumericPrice);
+		}
+		
+		for (NSString *invalidProduct in Response.invalidProductIdentifiers)
+		{
+			UE_LOG(LogOnline, Warning, TEXT("Problem in iTunes connect configuration for product: %s"), *FString(invalidProduct));
+		}
+		
+		CompletionDelegate.ExecuteIfBound(bWasSuccessful, OfferIds, TEXT(""));
+		bIsQueryInFlight = false;
+	}
+}
+
+void FOnlineStoreIOS::OnTransactionCompleteResponse(EPurchaseTransactionState Result, const FStoreKitTransactionData& TransactionData)
+{
+	UE_LOG(LogOnline, Verbose, TEXT("FOnlineStoreIOS::OnTransactionCompleteResponse"));
+}
+
+void FOnlineStoreIOS::OnTransactionRestored(const FStoreKitTransactionData& TransactionData)
+{
+	UE_LOG(LogOnline, Verbose, TEXT("FOnlineStoreIOS::OnTransactionRestored"));
+}
+
+void FOnlineStoreIOS::OnRestoreTransactionsComplete(EPurchaseTransactionState Result)
+{
+	UE_LOG(LogOnline, Verbose, TEXT("FOnlineStoreIOS::OnRestoreTransactionsComplete"));	
+}
+
+void FOnlineStoreIOS::GetOffers(TArray<FOnlineStoreOfferRef>& OutOffers) const
+{
+	for (const auto& CachedEntry : CachedOffers)
+	{
+		const FOnlineStoreOfferIOS& CachedOffer = CachedEntry.Value;
+		if (CachedOffer.IsValid())
+		{
+			OutOffers.Add(CachedOffer.Offer.ToSharedRef());
+		}
+	}
+}
+
+TSharedPtr<FOnlineStoreOffer> FOnlineStoreIOS::GetOffer(const FUniqueOfferId& OfferId) const
+{
+	TSharedPtr<FOnlineStoreOffer> Result;
+
+	const FOnlineStoreOfferIOS* Existing = CachedOffers.Find(OfferId);
+	if (Existing != nullptr && Existing->IsValid())
+	{
+		Result = Existing->Offer;
+	}
+	
+	return Result;
+}
+
+SKProduct* FOnlineStoreIOS::GetSKProductByOfferId(const FUniqueOfferId& OfferId)
+{
+	const FOnlineStoreOfferIOS* Existing = CachedOffers.Find(OfferId);
+	if (Existing != nullptr && Existing->IsValid())
+	{
+		return Existing->Product;
+	}
+	
+	return nil;
+}
\ No newline at end of file
diff --git a/Engine/Plugins/Online/IOS/OnlineSubsystemIOS/Source/Private/OnlineStoreInterfaceIOS.cpp b/Engine/Plugins/Online/IOS/OnlineSubsystemIOS/Source/Private/OnlineStoreInterfaceIOS.cpp
--- a/Engine/Plugins/Online/IOS/OnlineSubsystemIOS/Source/Private/OnlineStoreInterfaceIOS.cpp
+++ b/Engine/Plugins/Online/IOS/OnlineSubsystemIOS/Source/Private/OnlineStoreInterfaceIOS.cpp
@@ -1,376 +1,14 @@
 // Copyright 1998-2017 Epic Games, Inc. All Rights Reserved.
 
 #include "OnlineSubsystemIOSPrivatePCH.h"
-#include "OnlineStoreInterface.h"
+#include "OnlineStoreInterfaceIOS.h"
 #include "OnlineSubsystemNames.h"
-
-////////////////////////////////////////////////////////////////////
-/// FStoreKitHelper implementation
-
-
-@implementation FStoreKitHelper
-@synthesize Request;
-@synthesize AvailableProducts;
-
-
-- (id)init
-{
-	self = [super init];
-	return self;
-}
-
-
-- (void)paymentQueue:(SKPaymentQueue *)queue updatedTransactions:(NSArray *)transactions
-{
-	UE_LOG(LogOnline, Log, TEXT( "FStoreKitHelper::updatedTransactions" ));
-    for (SKPaymentTransaction *transaction in transactions)
-    {
-        switch ([transaction transactionState])
-        {
-            case SKPaymentTransactionStatePurchased:
-                [self completeTransaction:transaction];
-                break;
-            case SKPaymentTransactionStateFailed:
-                [self failedTransaction:transaction];
-                break;
-            case SKPaymentTransactionStateRestored:
-                [self restoreTransaction:transaction];
-				break;
-            case SKPaymentTransactionStatePurchasing:
-                continue;
-			default:
-				break;
-        }
-    }
-}
-
--(void)paymentQueueRestoreCompletedTransactionsFinished:(SKPaymentQueue *)queue
-{
-	UE_LOG(LogOnline, Log, TEXT("FStoreKitHelper::paymentQueueRestoreCompletedTransactionsFinished"));
-
-	[FIOSAsyncTask CreateTaskWithBlock : ^ bool(void)
-	{
-		IOnlineSubsystem* OnlineSub = IOnlineSubsystem::Get(IOS_SUBSYSTEM);
-		FOnlineStoreInterfaceIOS* StoreInterface = (FOnlineStoreInterfaceIOS*)OnlineSub->GetStoreInterface().Get();
-
-		if (StoreInterface->CachedPurchaseRestoreObject.IsValid())
-		{
-			StoreInterface->CachedPurchaseRestoreObject->ReadState = EOnlineAsyncTaskState::Done;
-		}
-		StoreInterface->TriggerOnInAppPurchaseRestoreCompleteDelegates(EInAppPurchaseState::Restored);
-		// Ensure these flags are false to allow a subsequent transaction attempt to proceed
-		StoreInterface->bIsProductRequestInFlight = false;
-		StoreInterface->bIsPurchasing = false;
-		StoreInterface->bIsRestoringPurchases = false;
-
-		return true;
-	}];
-}
-
--(void)paymentQueue: (SKPaymentQueue *)queue restoreCompletedTransactionsFailedWithError : (NSError *)error
-{
-	UE_LOG(LogOnline, Log, TEXT("FStoreKitHelper::restoreCompletedTransactionsFailedWithError - %s"), *FString([error localizedDescription]));
-
-	EInAppPurchaseState::Type CompletionState = EInAppPurchaseState::Unknown;
-	switch (error.code)
-	{
-	case SKErrorPaymentCancelled:
-		CompletionState = EInAppPurchaseState::Cancelled;
-		break;
-	case SKErrorClientInvalid:
-	case SKErrorStoreProductNotAvailable:
-	case SKErrorPaymentInvalid:
-		CompletionState = EInAppPurchaseState::Invalid;
-		break;
-	case SKErrorPaymentNotAllowed:
-		CompletionState = EInAppPurchaseState::NotAllowed;
-		break;
-	}
-
-	[FIOSAsyncTask CreateTaskWithBlock : ^ bool(void)
-	{
-		IOnlineSubsystem* OnlineSub = IOnlineSubsystem::Get(IOS_SUBSYSTEM);
-		FOnlineStoreInterfaceIOS* StoreInterface = (FOnlineStoreInterfaceIOS*)OnlineSub->GetStoreInterface().Get();
-		if (StoreInterface->CachedPurchaseRestoreObject.IsValid())
-		{
-			StoreInterface->CachedPurchaseRestoreObject->ReadState = EOnlineAsyncTaskState::Done;
-		}
-
-		StoreInterface->TriggerOnInAppPurchaseRestoreCompleteDelegates(CompletionState);
-		// Ensure these flags are false to allow a subsequent transaction attempt to proceed
-		StoreInterface->bIsProductRequestInFlight = false;
-		StoreInterface->bIsPurchasing = false;
-		StoreInterface->bIsRestoringPurchases = false;
-
-		return true;
-	}];
-}
-
--(void)paymentQueue: (SKPaymentQueue *)queue requestProductDataFailedWithError : (NSError *)error
-{
-	UE_LOG(LogOnline, Log, TEXT("FStoreKitHelper::requestProductDataFailedWithError - %s"), *FString([error localizedDescription]));
-
-	[FIOSAsyncTask CreateTaskWithBlock : ^ bool(void)
-	{
-		IOnlineSubsystem* OnlineSub = IOnlineSubsystem::Get(IOS_SUBSYSTEM);
-		FOnlineStoreInterfaceIOS* StoreInterface = (FOnlineStoreInterfaceIOS*)OnlineSub->GetStoreInterface().Get();
-		
-		StoreInterface->TriggerOnQueryForAvailablePurchasesCompleteDelegates( false );
-		// Ensure these flags are false to allow a subsequent transaction attempt to proceed
-		StoreInterface->bIsProductRequestInFlight = false;
-		StoreInterface->bIsPurchasing = false;
-		StoreInterface->bIsRestoringPurchases = false;
-
-		return true;
-	}];
-}
-
-- (void) completeTransaction: (SKPaymentTransaction *)transaction
-{
-	UE_LOG(LogOnline, Log, TEXT( "FStoreKitHelper::completeTransaction" ));
-
-    
-	[FIOSAsyncTask CreateTaskWithBlock : ^ bool(void)
-	{
-		IOnlineSubsystem* OnlineSub = IOnlineSubsystem::Get(IOS_SUBSYSTEM);
-		FOnlineStoreInterfaceIOS* StoreInterface = (FOnlineStoreInterfaceIOS*)OnlineSub->GetStoreInterface().Get();
-
-		if (StoreInterface->CachedPurchaseStateObject.IsValid())
-		{
-			FString ReceiptData;
-
-#ifdef __IPHONE_7_0
-			if( [IOSAppDelegate GetDelegate].OSVersion >= 7.0 )
-			{
-				NSURL* nsReceiptUrl = [[NSBundle mainBundle] appStoreReceiptURL];
-				NSData* nsReceiptData = [NSData dataWithContentsOfURL : nsReceiptUrl];
-				NSString* nsEncodedReceiptData = [nsReceiptData base64EncodedStringWithOptions : NSDataBase64EncodingEndLineWithLineFeed];
-     
-				ReceiptData = nsEncodedReceiptData;
-			}
-			else
-#endif
-			{
-#if __IPHONE_OS_VERSION_MIN_REQUIRED < __IPHONE_7_0
-				// If earlier than IOS 7, we will need to use the transactionReceipt
-				NSString* nsEncodedReceiptData = [transaction.transactionReceipt base64EncodedStringWithOptions:NSDataBase64EncodingEndLineWithLineFeed];
-				ReceiptData = nsEncodedReceiptData;
-#endif
-			}
-
-			StoreInterface->CachedPurchaseStateObject->ProvidedProductInformation.ReceiptData = ReceiptData;
-			StoreInterface->CachedPurchaseStateObject->ProvidedProductInformation.TransactionIdentifier = transaction.transactionIdentifier;
-			StoreInterface->CachedPurchaseStateObject->ReadState = EOnlineAsyncTaskState::Done;
-		}
-     
-		StoreInterface->TriggerOnInAppPurchaseCompleteDelegates(EInAppPurchaseState::Success);
-		// Ensure these flags are false to allow a subsequent transaction attempt to proceed
- 		StoreInterface->bIsProductRequestInFlight = false;
-		StoreInterface->bIsPurchasing = false;
-		StoreInterface->bIsRestoringPurchases = false;
-	 
-		return true;
-	}];
- 
-	// Remove the transaction from the payment queue.
-	[[SKPaymentQueue defaultQueue] finishTransaction: transaction];
-}
-
-
-- (void) restoreTransaction: (SKPaymentTransaction *)transaction
-{
-	UE_LOG(LogOnline, Log, TEXT( "FStoreKitHelper::restoreTransaction" ));
-
-	[FIOSAsyncTask CreateTaskWithBlock : ^ bool(void)
-	{
-		IOnlineSubsystem* OnlineSub = IOnlineSubsystem::Get(IOS_SUBSYSTEM);
-		FOnlineStoreInterfaceIOS* StoreInterface = (FOnlineStoreInterfaceIOS*)OnlineSub->GetStoreInterface().Get();
-
-		if (StoreInterface->CachedPurchaseRestoreObject.IsValid())
-		{
-			FString ReceiptData;
-
-#ifdef __IPHONE_7_0
-			if ([IOSAppDelegate GetDelegate].OSVersion >= 7.0)
-			{
-				NSURL* nsReceiptUrl = [[NSBundle mainBundle] appStoreReceiptURL];
-				NSData* nsReceiptData = [NSData dataWithContentsOfURL : nsReceiptUrl];
-				NSString* nsEncodedReceiptData = [nsReceiptData base64EncodedStringWithOptions : NSDataBase64EncodingEndLineWithLineFeed];
-
-				ReceiptData = nsEncodedReceiptData;
-			}
-			else
-#endif
-			{
-#if __IPHONE_OS_VERSION_MIN_REQUIRED < __IPHONE_7_0
-				// If earlier than IOS 7, we will need to use the transactionReceipt
-				NSString* nsEncodedReceiptData = [transaction.originalTransaction.transactionReceipt base64EncodedStringWithOptions : NSDataBase64EncodingEndLineWithLineFeed];
-				ReceiptData = nsEncodedReceiptData;
-#endif
-			}
-
-			FInAppPurchaseRestoreInfo RestoreInfo;
-			RestoreInfo.Identifier = FString(transaction.originalTransaction.payment.productIdentifier);
-			RestoreInfo.ReceiptData = ReceiptData;
-			StoreInterface->CachedPurchaseRestoreObject->ProvidedRestoreInformation.Add(RestoreInfo);
-		}
-
-		return true;
-	}];
-	
-	[[SKPaymentQueue defaultQueue] finishTransaction:transaction];
-}
-
-
-- (void) failedTransaction: (SKPaymentTransaction *)transaction
-{
-	UE_LOG(LogOnline, Log, TEXT( "FStoreKitHelper::failedTransaction - %s" ), *FString([transaction.error localizedDescription]));
-
-	EInAppPurchaseState::Type CompletionState = EInAppPurchaseState::Unknown;
-	switch (transaction.error.code)
-	{
-	case SKErrorPaymentCancelled:
-		CompletionState = EInAppPurchaseState::Cancelled;
-		break;
-    case SKErrorClientInvalid:
-		CompletionState = EInAppPurchaseState::AlreadyOwned;
-		break;
-	case SKErrorStoreProductNotAvailable:
-	case SKErrorPaymentInvalid:
-		CompletionState = EInAppPurchaseState::Invalid;
-		break;
-	case SKErrorPaymentNotAllowed:
-		CompletionState = EInAppPurchaseState::NotAllowed;
-		break;
-	}
-	
-    [FIOSAsyncTask CreateTaskWithBlock : ^ bool(void)
-    {
-        IOnlineSubsystem* OnlineSub = IOnlineSubsystem::Get(IOS_SUBSYSTEM);
-        FOnlineStoreInterfaceIOS* StoreInterface = (FOnlineStoreInterfaceIOS*)OnlineSub->GetStoreInterface().Get();
-        if (StoreInterface->CachedPurchaseStateObject.IsValid())
-        {
-            StoreInterface->CachedPurchaseStateObject->ReadState = EOnlineAsyncTaskState::Done;
-        }
-
-        StoreInterface->TriggerOnInAppPurchaseCompleteDelegates(CompletionState);
- 		// Ensure these flags are false to allow a subsequent transaction attempt to proceed
-		StoreInterface->bIsProductRequestInFlight = false;
-		StoreInterface->bIsPurchasing = false;
-		StoreInterface->bIsRestoringPurchases = false;
-	 
-        return true;
-     }];
-
-    [[SKPaymentQueue defaultQueue] finishTransaction: transaction];
-}
-
-
-- (void) requestProductData: (NSMutableSet*) productIDs
-{
-	UE_LOG(LogOnline, Verbose, TEXT( "FStoreKitHelper::requestProductData" ));
-
-	Request = [[SKProductsRequest alloc] initWithProductIdentifiers: productIDs];
-	Request.delegate = self;
-
-	[Request start];
-}
-
-
-- (void) makePurchase: (NSMutableSet*) productIDs
-{
-	UE_LOG(LogOnline, Verbose, TEXT( "FStoreKitHelper::makePurchase" ));
-
-	Request = [[SKProductsRequest alloc] initWithProductIdentifiers: productIDs];
-	Request.delegate = self;
-
-	[Request start];
-}
-
-
-- (void)productsRequest: (SKProductsRequest *)request didReceiveResponse:(SKProductsResponse *)response
-{
-	UE_LOG(LogOnline, Verbose, TEXT( "FStoreKitHelper::didReceiveResponse" ));
-	// Direct the response back to the store interface
-	[FIOSAsyncTask CreateTaskWithBlock : ^ bool(void)
-	{
-		IOnlineSubsystem* OnlineSub = IOnlineSubsystem::Get(IOS_SUBSYSTEM);
-		FOnlineStoreInterfaceIOS* StoreInterface = (FOnlineStoreInterfaceIOS*)OnlineSub->GetStoreInterface().Get();
-		StoreInterface->ProcessProductsResponse(response);
-
-		return true;
-	}];
-
-	[request autorelease];
-}
-
--(void)requestDidFinish:(SKRequest*)request
-{
-#ifdef __IPHONE_7_0
-	if ([Request isKindOfClass : [SKReceiptRefreshRequest class]])
-	{
-		[[SKPaymentQueue defaultQueue] restoreCompletedTransactions];
-	}
-	
-	else if ([Request isKindOfClass : [SKProductsRequest class]])
-	{
-		// This is a placeholder to check if this case needs to be handled in the future
-	}
-	
-#endif
-}
-
--(void)request:(SKRequest*)request didFailWithError: (NSError *)error
-{
-#ifdef __IPHONE_7_0
-	// Check if we were refreshing a receipt
-	if ([Request isKindOfClass : [SKReceiptRefreshRequest class]])
-	{
-		UE_LOG(LogOnline, Warning, TEXT( "FStoreKitHelper: SKReceiptRefreshRequest failed" ));
-		[self paymentQueue:[SKPaymentQueue defaultQueue] restoreCompletedTransactionsFailedWithError: error];
-		[Request release];
-	}
-
-	// Check if we are doing a product info request
-	else if ([Request isKindOfClass : [SKProductsRequest class]])
-	{
-		UE_LOG(LogOnline, Warning, TEXT( "FStoreKitHelper: SKProductsRequest failed" ));
-		[self paymentQueue:[SKPaymentQueue defaultQueue] requestProductDataFailedWithError: error];
-		[Request release];
-	}
-	else
-	{
-		UE_LOG(LogOnline, Warning, TEXT( "FStoreKitHelper: didFailWithError went unhandled" ));
-	}
-	
-#endif
-}
-
--(void)restorePurchases
-{
-#ifdef __IPHONE_7_0
-	if ([IOSAppDelegate GetDelegate].OSVersion >= 7.0)
-	{
-		Request = [[SKReceiptRefreshRequest alloc] init];
-		Request.delegate = self;
-		[Request start];
-    }
-	else
-#endif
-	{
-#if __IPHONE_OS_VERSION_MIN_REQUIRED < __IPHONE_7_0
-		[[SKPaymentQueue defaultQueue] restoreCompletedTransactions];
-#endif
-	}
-}
-@end
-
+#import "OnlineStoreKitHelper.h"
 
 
 ////////////////////////////////////////////////////////////////////
 /// FOnlineStoreInterfaceIOS implementation
 
-
 FOnlineStoreInterfaceIOS::FOnlineStoreInterfaceIOS() 
 {
 	UE_LOG(LogOnline, Verbose, TEXT( "FOnlineStoreInterfaceIOS::FOnlineStoreInterfaceIOS" ));
diff --git a/Engine/Plugins/Online/IOS/OnlineSubsystemIOS/Source/Private/OnlineStoreInterfaceIOS.h b/Engine/Plugins/Online/IOS/OnlineSubsystemIOS/Source/Private/OnlineStoreInterfaceIOS.h
--- a/Engine/Plugins/Online/IOS/OnlineSubsystemIOS/Source/Private/OnlineStoreInterfaceIOS.h
+++ b/Engine/Plugins/Online/IOS/OnlineSubsystemIOS/Source/Private/OnlineStoreInterfaceIOS.h
@@ -4,16 +4,8 @@
 
 #include "OnlineStoreInterface.h"
 
-// SK includes
-#include <StoreKit/SKRequest.h>
-#include <StoreKit/SKError.h>
-#include <StoreKit/SKProduct.h>
-#include <StoreKit/SKProductsRequest.h>
-#include <StoreKit/SKReceiptRefreshRequest.h>
-#include <StoreKit/SKPaymentTransaction.h>
-#include <StoreKit/SKPayment.h>
-#include <StoreKit/SKPaymentQueue.h>
-
+@class FStoreKitHelper;
+@class SKProductsResponse;
 
 /**
  * The resulting state of an iap transaction
@@ -29,29 +21,6 @@ namespace EInAppPurchaseResult
     };
 }
 
-
-/** Helper class, which allows us to manage IAP product information requests, AND transactions */
-@interface FStoreKitHelper : NSObject<SKProductsRequestDelegate, SKPaymentTransactionObserver, SKRequestDelegate>
-{
-};
-/** Store kit request object, holds information about the products we are purchasing, or querying. */
-@property (nonatomic, strong) SKRequest *Request;
-/** collection of available products attaced through a store kit request */
-@property (nonatomic, strong) NSArray *AvailableProducts;
-
-/** Helper fn to start a store kit purchase request */
--(void)makePurchase:(NSMutableSet*)productIDs;
-/** Helper fn to start a store kit purchase information query request */
-- (void) requestProductData:(NSMutableSet*) productIDs;
-
-/** Helper fn to direct a product request response back to our store interface */
-- (void)productsRequest:(SKProductsRequest *)request didReceiveResponse:(SKProductsResponse *)response;
-
-/** Helper fn to restore previously purchased products */
--(void)restorePurchases;
-@end
-
-
 /**
  *	FOnlineStoreInterfaceIOS - Implementation of the online store for IOS
  */
