Update list of supported platforms for Engine copy of OnlineSubsystemGooglePlay, so it's automatically excluded on other platforms. Currently preventing Match3 sample from packaging.

diff --git a/Engine/Plugins/Online/Android/OnlineSubsystemGooglePlay/OnlineSubsystemGooglePlay.uplugin b/Engine/Plugins/Online/Android/OnlineSubsystemGooglePlay/OnlineSubsystemGooglePlay.uplugin
--- a/Engine/Plugins/Online/Android/OnlineSubsystemGooglePlay/OnlineSubsystemGooglePlay.uplugin
+++ b/Engine/Plugins/Online/Android/OnlineSubsystemGooglePlay/OnlineSubsystemGooglePlay.uplugin
@@ -9,7 +9,9 @@
 	"CreatedBy" : "Epic Games, Inc.",
 	"CreatedByURL" : "http://epicgames.com",
 	"EnabledByDefault" : true,
-
+	"SupportedTargetPlatforms": [
+		"Android"
+	],
     "Modules": [
         {
             "Name": "OnlineSubsystemGooglePlay",
