Fix Firebase resource staging and harden GooglePlay login error handling

#android

diff --git a/Engine/Plugins/Online/Android/OnlineSubsystemGooglePlay/Source/OnlineSubsystemGooglePlay_UPL.xml b/Engine/Plugins/Online/Android/OnlineSubsystemGooglePlay/Source/OnlineSubsystemGooglePlay_UPL.xml
--- a/Engine/Plugins/Online/Android/OnlineSubsystemGooglePlay/Source/OnlineSubsystemGooglePlay_UPL.xml
+++ b/Engine/Plugins/Online/Android/OnlineSubsystemGooglePlay/Source/OnlineSubsystemGooglePlay_UPL.xml
@@ -8,15 +8,21 @@
 
 	<!-- optional additions to GameActivity onCreate in GameActivity.java -->
 	<gameActivityOnCreateAdditions>
-	  <insert>
-		//Google Play SDK onCreate additions
-		googleClient = new GoogleApiClient.Builder(this)
-		.addConnectionCallbacks(this)
-		.addOnConnectionFailedListener(this)
-		.addApi(Games.API).addScope(Games.SCOPE_GAMES)
-		.addApi(Plus.API).addScope(Plus.SCOPE_PLUS_LOGIN)
-		.build();
-	  </insert>
+		<insert>
+			//Google Play SDK onCreate additions
+			try {
+				GoogleApiClient.Builder gbuilder = new GoogleApiClient.Builder(this);
+				gbuilder.addConnectionCallbacks(this);
+				gbuilder.addOnConnectionFailedListener(this);
+				gbuilder.addApiIfAvailable(Games.API, Games.SCOPE_GAMES);
+				googleClient = gbuilder.build();
+			}
+			catch (Exception e)
+			{
+				Log.debug("GoogleApiClient exception caught: " + e.toString());
+			}
+			Log.debug("googleClient is " + ((googleClient == null) ? "disabled" : "valid"));
+		</insert>
 	</gameActivityOnCreateAdditions>
 
 </root>
diff --git a/Engine/Plugins/Runtime/Firebase/Source/Firebase.upl.xml b/Engine/Plugins/Runtime/Firebase/Source/Firebase.upl.xml
--- a/Engine/Plugins/Runtime/Firebase/Source/Firebase.upl.xml
+++ b/Engine/Plugins/Runtime/Firebase/Source/Firebase.upl.xml
@@ -67,7 +67,11 @@ apply plugin: 'com.google.gms.google-services'
 	</buildGradleAdditions>
 
 	<prebuildCopies>
-		<copyDir src="$S(PluginDir)/Java" dst="$S(BuildDir)/src/com/epicgames/ue4"/>
+		<if condition="bEnabled">
+			<true>
+				<copyDir src="$S(PluginDir)/Java" dst="$S(BuildDir)/src/com/epicgames/ue4"/>
+			</true>
+		</if>
 	</prebuildCopies>
 	
 	<gradleCopies>
