EnableSnapshot for Google Play Games login (contributed by [REDACTED_PERSON])

#PR #4805
#ue4
#android

diff --git a/Engine/Plugins/Online/Android/OnlineSubsystemGooglePlay/Source/Private/OnlineAsyncTaskGooglePlayLogin.cpp b/Engine/Plugins/Online/Android/OnlineSubsystemGooglePlay/Source/Private/OnlineAsyncTaskGooglePlayLogin.cpp
--- a/Engine/Plugins/Online/Android/OnlineSubsystemGooglePlay/Source/Private/OnlineAsyncTaskGooglePlayLogin.cpp
+++ b/Engine/Plugins/Online/Android/OnlineSubsystemGooglePlay/Source/Private/OnlineAsyncTaskGooglePlayLogin.cpp
@@ -42,6 +42,7 @@ void FOnlineAsyncTaskGooglePlayLogin::Start_OnTaskThread()
 					*FString(DebugString(LocalStatus).c_str()));
 				LocalSubsystem->OnAuthActionFinished(Op, LocalStatus);
 			})
+			.EnableSnapshots()
 			.AddOauthScope("https://www.googleapis.com/auth/plus.login")
 			.Create(Subsystem->PlatformConfiguration);
 	}
