Add checkbox to enable snapshot on Google Games

#ue4
#android

diff --git a/Engine/Config/BaseEngine.ini b/Engine/Config/BaseEngine.ini
--- a/Engine/Config/BaseEngine.ini
+++ b/Engine/Config/BaseEngine.ini
@@ -1885,6 +1885,7 @@ bBuildWithHiddenSymbolVisibility=false
 bSaveSymbols=false
 bAllowControllers=True
 bAllowIMU=True
+bEnableSnapshots=False
 
 [/Script/AndroidPlatformEditor.AndroidSDKSettings]
 SDKAPILevel=latest
diff --git a/Engine/Plugins/Online/Android/OnlineSubsystemGooglePlay/Source/Private/OnlineAsyncTaskGooglePlayLogin.cpp b/Engine/Plugins/Online/Android/OnlineSubsystemGooglePlay/Source/Private/OnlineAsyncTaskGooglePlayLogin.cpp
--- a/Engine/Plugins/Online/Android/OnlineSubsystemGooglePlay/Source/Private/OnlineAsyncTaskGooglePlayLogin.cpp
+++ b/Engine/Plugins/Online/Android/OnlineSubsystemGooglePlay/Source/Private/OnlineAsyncTaskGooglePlayLogin.cpp
@@ -2,6 +2,7 @@
 
 #include "OnlineAsyncTaskGooglePlayLogin.h"
 #include "OnlineSubsystemGooglePlay.h"
+#include "AndroidRuntimeSettings.h"
 
 THIRD_PARTY_INCLUDES_START
 #include "gpg/builder.h"
@@ -31,7 +32,9 @@ void FOnlineAsyncTaskGooglePlayLogin::Start_OnTaskThread()
 		// Store the Subsystem pointer locally so that the OnAuthActionFinished lambda can capture it
 		FOnlineSubsystemGooglePlay* LocalSubsystem = Subsystem;
 
-		Subsystem->GameServicesPtr = gpg::GameServices::Builder()
+		auto DefaultSettings = GetDefault<UAndroidRuntimeSettings>();
+
+		gpg::GameServices::Builder builder = gpg::GameServices::Builder()
 			.SetDefaultOnLog(gpg::LogLevel::VERBOSE)
 			.SetOnAuthActionStarted([](gpg::AuthOperation Op) {
 				UE_LOG_ONLINE(Log, TEXT("GPG OnAuthActionStarted: %s"), *FString(DebugString(Op).c_str()));
@@ -41,10 +44,14 @@ void FOnlineAsyncTaskGooglePlayLogin::Start_OnTaskThread()
 					*FString(DebugString(Op).c_str()),
 					*FString(DebugString(LocalStatus).c_str()));
 				LocalSubsystem->OnAuthActionFinished(Op, LocalStatus);
-			})
-			.EnableSnapshots()
-			.AddOauthScope("https://www.googleapis.com/auth/plus.login")
-			.Create(Subsystem->PlatformConfiguration);
+			});
+			
+		if (DefaultSettings->bEnableSnapshots)
+		{
+			builder.EnableSnapshots();
+		}
+		
+		Subsystem->GameServicesPtr = builder.Create(Subsystem->PlatformConfiguration);
 	}
 	else if(Subsystem->GameServicesPtr->IsAuthorized())
 	{
diff --git a/Engine/Source/Runtime/Android/AndroidRuntimeSettings/Classes/AndroidRuntimeSettings.h b/Engine/Source/Runtime/Android/AndroidRuntimeSettings/Classes/AndroidRuntimeSettings.h
--- a/Engine/Source/Runtime/Android/AndroidRuntimeSettings/Classes/AndroidRuntimeSettings.h
+++ b/Engine/Source/Runtime/Android/AndroidRuntimeSettings/Classes/AndroidRuntimeSettings.h
@@ -388,6 +388,10 @@ public:
 	UPROPERTY(GlobalConfig, EditAnywhere, Category = GooglePlayServices)
 	TArray<FGooglePlayLeaderboardMapping> LeaderboardMap;
 
+	// Enabling this requests snapshots support for saved games during Google Play login.
+	UPROPERTY(GlobalConfig, EditAnywhere, Category = GooglePlayServices, meta = (DisplayName = "Enable Snapshots on Google Play login"))
+	bool bEnableSnapshots;
+
 	// Enabling this includes the AdMob SDK and will be detected by Google Play Console on upload of APK.  Disable if you do not need ads to remove this warning.
 	UPROPERTY(GlobalConfig, EditAnywhere, Category = GooglePlayServices, meta = (DisplayName = "Include AdMob support for ads"))
 	bool bSupportAdMob;
