Fix compile issue

#ue4
#android

diff --git a/Engine/Plugins/Online/Android/OnlineSubsystemGooglePlay/Source/Private/OnlineAsyncTaskGooglePlayLogin.cpp b/Engine/Plugins/Online/Android/OnlineSubsystemGooglePlay/Source/Private/OnlineAsyncTaskGooglePlayLogin.cpp
--- a/Engine/Plugins/Online/Android/OnlineSubsystemGooglePlay/Source/Private/OnlineAsyncTaskGooglePlayLogin.cpp
+++ b/Engine/Plugins/Online/Android/OnlineSubsystemGooglePlay/Source/Private/OnlineAsyncTaskGooglePlayLogin.cpp
@@ -34,24 +34,37 @@ void FOnlineAsyncTaskGooglePlayLogin::Start_OnTaskThread()
 
 		auto DefaultSettings = GetDefault<UAndroidRuntimeSettings>();
 
-		gpg::GameServices::Builder builder = gpg::GameServices::Builder()
-			.SetDefaultOnLog(gpg::LogLevel::VERBOSE)
-			.SetOnAuthActionStarted([](gpg::AuthOperation Op) {
-				UE_LOG_ONLINE(Log, TEXT("GPG OnAuthActionStarted: %s"), *FString(DebugString(Op).c_str()));
-			})
-			.SetOnAuthActionFinished([LocalSubsystem](gpg::AuthOperation Op, gpg::AuthStatus LocalStatus) {
-				UE_LOG_ONLINE(Log, TEXT("GPG OnAuthActionFinished: %s, AuthStatus: %s"),
-					*FString(DebugString(Op).c_str()),
-					*FString(DebugString(LocalStatus).c_str()));
-				LocalSubsystem->OnAuthActionFinished(Op, LocalStatus);
-			});
-			
 		if (DefaultSettings->bEnableSnapshots)
 		{
-			builder.EnableSnapshots();
+			Subsystem->GameServicesPtr = gpg::GameServices::Builder()
+				.SetDefaultOnLog(gpg::LogLevel::VERBOSE)
+				.SetOnAuthActionStarted([](gpg::AuthOperation Op) {
+					UE_LOG_ONLINE(Log, TEXT("GPG OnAuthActionStarted: %s"), *FString(DebugString(Op).c_str()));
+				})
+				.SetOnAuthActionFinished([LocalSubsystem](gpg::AuthOperation Op, gpg::AuthStatus LocalStatus) {
+					UE_LOG_ONLINE(Log, TEXT("GPG OnAuthActionFinished: %s, AuthStatus: %s"),
+						*FString(DebugString(Op).c_str()),
+						*FString(DebugString(LocalStatus).c_str()));
+					LocalSubsystem->OnAuthActionFinished(Op, LocalStatus);
+				})
+				.EnableSnapshots()
+				.Create(Subsystem->PlatformConfiguration);
+		}
+		else
+		{
+			Subsystem->GameServicesPtr = gpg::GameServices::Builder()
+				.SetDefaultOnLog(gpg::LogLevel::VERBOSE)
+				.SetOnAuthActionStarted([](gpg::AuthOperation Op) {
+					UE_LOG_ONLINE(Log, TEXT("GPG OnAuthActionStarted: %s"), *FString(DebugString(Op).c_str()));
+				})
+				.SetOnAuthActionFinished([LocalSubsystem](gpg::AuthOperation Op, gpg::AuthStatus LocalStatus) {
+					UE_LOG_ONLINE(Log, TEXT("GPG OnAuthActionFinished: %s, AuthStatus: %s"),
+						*FString(DebugString(Op).c_str()),
+						*FString(DebugString(LocalStatus).c_str()));
+					LocalSubsystem->OnAuthActionFinished(Op, LocalStatus);
+				})
+				.Create(Subsystem->PlatformConfiguration);
 		}
-		
-		Subsystem->GameServicesPtr = builder.Create(Subsystem->PlatformConfiguration);
 	}
 	else if(Subsystem->GameServicesPtr->IsAuthorized())
 	{
