Fixup of more merge/build errors

diff --git a/Engine/Plugins/Online/Android/OnlineSubsystemGooglePlay/Source/Private/OnlineStoreGooglePlay.cpp b/Engine/Plugins/Online/Android/OnlineSubsystemGooglePlay/Source/Private/OnlineStoreGooglePlay.cpp
--- a/Engine/Plugins/Online/Android/OnlineSubsystemGooglePlay/Source/Private/OnlineStoreGooglePlay.cpp
+++ b/Engine/Plugins/Online/Android/OnlineSubsystemGooglePlay/Source/Private/OnlineStoreGooglePlay.cpp
@@ -50,7 +50,7 @@ void FOnlineStoreGooglePlayV2::Init()
 	AndroidThunkCpp_Iap_SetupIapService(GooglePlayLicenseKey);
 }
 
-void FOnlineStoreGooglePlayV2::OnGooglePlayAvailableIAPQueryComplete(EGooglePlayBillingResponseCode InResponseCode, const TArray<FInAppPurchaseProductInfo>& InProvidedProductInformation)
+void FOnlineStoreGooglePlayV2::OnGooglePlayAvailableIAPQueryComplete(EGooglePlayBillingResponseCode InResponseCode, const TArray<FOnlineStoreOffer>& InProvidedProductInformation)
 { 
 	TArray<FUniqueOfferId> OfferIds;
 	for (const FOnlineStoreOffer& Product : InProvidedProductInformation)
diff --git a/Engine/Plugins/Online/OnlineSubsystemUtils/Source/OnlineBlueprintSupport/Private/K2Node_InAppPurchaseUnprocessed2.cpp b/Engine/Plugins/Online/OnlineSubsystemUtils/Source/OnlineBlueprintSupport/Private/K2Node_InAppPurchaseUnprocessed2.cpp
--- a/Engine/Plugins/Online/OnlineSubsystemUtils/Source/OnlineBlueprintSupport/Private/K2Node_InAppPurchaseUnprocessed2.cpp
+++ b/Engine/Plugins/Online/OnlineSubsystemUtils/Source/OnlineBlueprintSupport/Private/K2Node_InAppPurchaseUnprocessed2.cpp
@@ -1,4 +1,4 @@
-// Copyright 1998-2019 Epic Games, Inc. All Rights Reserved.
+// Copyright Epic Games, Inc. All Rights Reserved.
 
 #include "K2Node_InAppPurchaseUnprocessed2.h"
 #include "InAppPurchaseCallbackProxy2.h"
