Added a "NotImplemented" implementation of the new Checkout method whose callback has no receipts for Googleplay

diff --git a/Engine/Plugins/Online/Android/OnlineSubsystemGooglePlay/Source/Private/OnlinePurchaseGooglePlay.cpp b/Engine/Plugins/Online/Android/OnlineSubsystemGooglePlay/Source/Private/OnlinePurchaseGooglePlay.cpp
--- a/Engine/Plugins/Online/Android/OnlineSubsystemGooglePlay/Source/Private/OnlinePurchaseGooglePlay.cpp
+++ b/Engine/Plugins/Online/Android/OnlineSubsystemGooglePlay/Source/Private/OnlinePurchaseGooglePlay.cpp
@@ -140,6 +140,13 @@ void FOnlinePurchaseGooglePlay::Checkout(const FUniqueNetId& UserId, const FPurc
 	}
 }
 
+void FOnlinePurchaseGooglePlay::Checkout(const FUniqueNetId& UserId, const FPurchaseCheckoutRequest& CheckoutRequest, const FOnPurchaseReceiptlessCheckoutComplete& Delegate)
+{
+	// Not Supported
+	UE_LOG_ONLINE(Error, TEXT("FOnlinePurchaseGooglePlay::Checkout (receiptless) is not currently supported"));
+	Delegate.ExecuteIfBound(FOnlineError(EOnlineErrorResult::NotImplemented));
+}
+
 void FOnlinePurchaseGooglePlay::FinalizePurchase(const FUniqueNetId& UserId, const FString& ReceiptId)
 {
 	UE_LOG_ONLINE_PURCHASE(Verbose, TEXT("FOnlinePurchaseGooglePlay::FinalizePurchase %s %s"), *UserId.ToString(), *ReceiptId);
diff --git a/Engine/Plugins/Online/Android/OnlineSubsystemGooglePlay/Source/Public/OnlinePurchaseGooglePlay.h b/Engine/Plugins/Online/Android/OnlineSubsystemGooglePlay/Source/Public/OnlinePurchaseGooglePlay.h
--- a/Engine/Plugins/Online/Android/OnlineSubsystemGooglePlay/Source/Public/OnlinePurchaseGooglePlay.h
+++ b/Engine/Plugins/Online/Android/OnlineSubsystemGooglePlay/Source/Public/OnlinePurchaseGooglePlay.h
@@ -130,6 +130,7 @@ public:
 
 	virtual bool IsAllowedToPurchase(const FUniqueNetId& UserId) override;
 	virtual void Checkout(const FUniqueNetId& UserId, const FPurchaseCheckoutRequest& CheckoutRequest, const FOnPurchaseCheckoutComplete& Delegate) override;
+	virtual void Checkout(const FUniqueNetId& UserId, const FPurchaseCheckoutRequest& CheckoutRequest, const FOnPurchaseReceiptlessCheckoutComplete& Delegate) override;
 	virtual void FinalizePurchase(const FUniqueNetId& UserId, const FString& ReceiptId) override;
 	virtual void RedeemCode(const FUniqueNetId& UserId, const FRedeemCodeRequest& RedeemCodeRequest, const FOnPurchaseRedeemCodeComplete& Delegate) override;
 	virtual void QueryReceipts(const FUniqueNetId& UserId, bool bRestoreReceipts, const FOnQueryReceiptsComplete& Delegate) override;
