Fix error in GooglePlay UPL

#android

diff --git a/Engine/Plugins/Online/Android/OnlineSubsystemGooglePlay/Source/OnlineSubsystemGooglePlay_UPL.xml b/Engine/Plugins/Online/Android/OnlineSubsystemGooglePlay/Source/OnlineSubsystemGooglePlay_UPL.xml
--- a/Engine/Plugins/Online/Android/OnlineSubsystemGooglePlay/Source/OnlineSubsystemGooglePlay_UPL.xml
+++ b/Engine/Plugins/Online/Android/OnlineSubsystemGooglePlay/Source/OnlineSubsystemGooglePlay_UPL.xml
@@ -13,13 +13,11 @@
 	
 	<!-- gradle dependencies additions -->
 	<buildGradleAdditions>
-		<true>
-			<insert>
-				dependencies {
-				implementation 'com.android.billingclient:billing:3.0.0'
-				}
-			</insert>
-		</true>
+		<insert>
+			dependencies {
+			implementation 'com.android.billingclient:billing:3.0.0'
+			}
+		</insert>
 	</buildGradleAdditions>
 
 	<gameActivityAndroidThunkJavaIapBeginPurchase>
@@ -62,27 +60,11 @@ public boolean AndroidThunkJava_IapBeginPurchase(String ProductId, String Accoun
 }
 		</insert>
 	</gameActivityAndroidThunkJavaIapBeginPurchase>
-		if (getPackageManager().checkPermission("com.android.vending.BILLING", getPackageName()) == getPackageManager().PERMISSION_GRANTED)
-		{
-			IapStoreHelper = new GooglePlayStoreHelper(InProductKey, this, Log);
-			if (IapStoreHelper != null)
-			{
-				Log.debug("[JAVA] - AndroidThunkJava_IapSetupService - Setup started");
-			}
-			else
-			{
-				Log.debug("[JAVA] - AndroidThunkJava_IapSetupService - Failed to setup IAP service");
-			}
-		}
-		else
-		{
-			Log.debug("[JAVA] - AndroidThunkJava_IapSetupService - You do not have the appropriate permission setup.");
-			Log.debug("[JAVA] - AndroidThunkJava_IapSetupService - Please ensure com.android.vending.BILLING is added to the manifest.");
-		}
-
+	
 	<!-- optional additions to GameActivity imports in GameActivity.java -->
 	<gameActivityImportAdditions>
 		<insert>
+import com.epicgames.unreal.GooglePlayStoreHelper;
 import java.io.UnsupportedEncodingException;
 import java.security.MessageDigest;
 import java.security.NoSuchAlgorithmException;
@@ -90,6 +72,18 @@ import java.security.NoSuchAlgorithmException;
 	</gameActivityImportAdditions>
 
 	<!-- optional additions to GameActivity class in GameActivity.java -->
+	<gameActivityClassAdditions>
+		<insert>
+	/** Implement this if app wants to handle IAB activity launch. **/
+	GooglePlayStoreHelper.PurchaseLaunchCallback purchaseLaunchCallback = null;
+
+	public GooglePlayStoreHelper.PurchaseLaunchCallback getPurchaseLaunchCallback()
+	{
+		return purchaseLaunchCallback;
+	}
+		</insert>
+	</gameActivityClassAdditions>
+
 	<gameActivityIapSetupServiceAdditions>
 		<insert>
 		if (getPackageManager().checkPermission("com.android.vending.BILLING", getPackageName()) == getPackageManager().PERMISSION_GRANTED)
