- Update Google Billing Library to v5 and fix or mitigate several issues found before the update

Included update and fixes together because the update was affecting how the fixes would be implemented

UE-160024:
 - Updated Google Billing Library to v5.
 - Many types and methods changed.
 - Cleaned up unused/non accessible/non needed methods and interfaces from GooglePlayStoreHelper.java, GameActivity.java.template and additions from OnlineSubsystemGooglePlay_UPL.xml
 - Used a uniform code style on GooglePlayStoreHelper.java (tried to be near the C++ guidelines)
 - Also follow C++ code guidelines and naming on the Java native functions implementation
 - Fixed support for GooglePlay pending transactions
 - General cleanup on OnlinePurchaseGooglePlay.h/cpp
 - Do not add validation information for pending purchases so they are not redeemable (thus not consumed accidentally)

UE-160021 and UE-160022
 - Use just one list of transactions to simplify handling them
 - Clear the list of transactions and fill it with the result of the query. GooglePlay always returns the full list of non consumed transactions (and the list of non consumable ones, but those are not currently supported by UE)
 - Added consume delegate to be able to remove consumed receipts

UE-160023:
 - Checked result of BillingClient.launchBillingFlow and notify the native side in case of failure

OI-3209:
 - When launch mode is set to single task in the manifest and the user moves to background and then to foreground using the Android launcher the onPurchasesUpdated listener reports BillingClient.BillingResponseCode.USER_CANCELLED with an empty list of purchases. Detect that as in progress purchase was cancelled

#9528

diff --git a/Engine/Build/Android/Java/src/com/epicgames/unreal/GameActivity.java.template b/Engine/Build/Android/Java/src/com/epicgames/unreal/GameActivity.java.template
--- a/Engine/Build/Android/Java/src/com/epicgames/unreal/GameActivity.java.template
+++ b/Engine/Build/Android/Java/src/com/epicgames/unreal/GameActivity.java.template
@@ -4726,7 +4726,6 @@ public class GameActivity extends $${gameActivitySuperClass}$$ implements Surfac
 				return;
 			}
 		}
-//$${gameActivityOnActivityResultIapStoreHelperHandler}$$
 		else
 		{
 			super.onActivityResult(requestCode, resultCode, data);
@@ -4787,22 +4786,6 @@ public class GameActivity extends $${gameActivitySuperClass}$$ implements Surfac
 		return bTriggeredQuery;
 	}
 
-	public boolean AndroidThunkJava_IapRestorePurchases(String[] InProductIDs, boolean[] bConsumable)
-	{
-		Log.debug("[JAVA] - AndroidThunkJava_IapRestorePurchases");
-		boolean bTriggeredRestore = false;
-		if( IapStoreHelper != null )
-		{
-			Log.debug("[JAVA] - AndroidThunkJava_IapRestorePurchases - Kick off logic here!");
-			bTriggeredRestore = IapStoreHelper.RestorePurchases(InProductIDs, bConsumable);
-		}
-		else
-		{
-			Log.debug("[JAVA] - Store Helper is invalid");
-		}
-		return bTriggeredRestore;
-	}
-
 	public void AndroidThunkJava_DismissSplashScreen()
 	{
 		if (mSplashDialog != null)
diff --git a/Engine/Build/Android/Java/src/com/epicgames/unreal/StoreHelper.java b/Engine/Build/Android/Java/src/com/epicgames/unreal/StoreHelper.java
--- a/Engine/Build/Android/Java/src/com/epicgames/unreal/StoreHelper.java
+++ b/Engine/Build/Android/Java/src/com/epicgames/unreal/StoreHelper.java
@@ -12,7 +12,5 @@ public interface StoreHelper
 	public boolean IsAllowedToMakePurchases();
 	public void ConsumePurchase(String purchaseToken);
 	public boolean QueryExistingPurchases();
-	public boolean RestorePurchases(String[] InProductIDs, boolean[] bConsumable);
 	public void	onDestroy();
-	public boolean onActivityResult(int requestCode, int resultCode, Intent data);
 }
\ No newline at end of file
diff --git a/Engine/Plugins/Online/Android/OnlineSubsystemGooglePlay/Source/Java/BillingApiV2/com/epicgames/unreal/GooglePlayStoreHelper.java b/Engine/Plugins/Online/Android/OnlineSubsystemGooglePlay/Source/Java/BillingApiV2/com/epicgames/unreal/GooglePlayStoreHelper.java
--- a/Engine/Plugins/Online/Android/OnlineSubsystemGooglePlay/Source/Java/BillingApiV2/com/epicgames/unreal/GooglePlayStoreHelper.java
+++ b/Engine/Plugins/Online/Android/OnlineSubsystemGooglePlay/Source/Java/BillingApiV2/com/epicgames/unreal/GooglePlayStoreHelper.java
@@ -1,116 +1,72 @@
 // Copyright Epic Games, Inc. All Rights Reserved.
-//This file needs to be here so the "ant" build step doesnt fail when looking for a /src folder.
 
 package com.epicgames.unreal;
 
-import android.app.Activity;
-import android.app.PendingIntent;
-import android.content.ComponentName;
-import android.content.Context;
-import android.content.Intent;
-import android.content.IntentSender.SendIntentException;
-import android.content.ServiceConnection;
-import android.os.Bundle;
-import android.os.Handler;
-import android.os.IBinder;
-import android.os.Looper;
-import android.os.RemoteException;
+import androidx.annotation.NonNull;
+import androidx.annotation.Nullable;
+
+import com.android.billingclient.api.BillingClient;
+import com.android.billingclient.api.BillingClientStateListener;
+import com.android.billingclient.api.BillingFlowParams;
+import com.android.billingclient.api.BillingResult;
+import com.android.billingclient.api.ConsumeParams;
+import com.android.billingclient.api.ConsumeResponseListener;
+import com.android.billingclient.api.ProductDetails;
+import com.android.billingclient.api.ProductDetailsResponseListener;
+import com.android.billingclient.api.Purchase;
+import com.android.billingclient.api.PurchasesResponseListener;
+import com.android.billingclient.api.PurchasesUpdatedListener;
+import com.android.billingclient.api.QueryProductDetailsParams;
+import com.android.billingclient.api.QueryPurchasesParams;
 import com.android.vending.billing.util.Base64; //WMM should use different base64 here.
-import com.android.billingclient.api.*;
+
 import java.util.ArrayList;
+import java.util.Arrays;
 import java.util.List;
-import org.json.JSONException;
-import org.json.JSONObject; 
 
-public class GooglePlayStoreHelper implements StoreHelper
+public class GooglePlayStoreHelper implements StoreHelper, PurchasesUpdatedListener
 {
 	// Our IAB helper interface provided by google.
-	private BillingClient mBillingClient;
+	private final BillingClient Client;
 
 	// Flag that determines whether the store is ready to use.
-	private boolean bIsIapSetup;
+	private boolean bIsIapSetup = false;
 
 	// Output device for log messages.
-	private Logger Log;
+	private final Logger Log;
 
 	// Cache access to the games activity.
-	private GameActivity gameActivity;
-
-	// The google play license key.
-	private String productKey;
-
-	private final int UndefinedFailureResponse = -1;
-
-	private int outstandingConsumeRequests = 0;
-
-	private boolean bIsConsumeComplete = false;
-
-	private class GooglePlayProductDescription
-	{
-		// Product offer id 
-		public String id;
-		// Product friendly name
-		public String title;
-		// Product description
-		public String description;
-		// Currency friendly string 
-		public String price;
-		// Raw price in currency units
-		public Float priceRaw;
-		// Local currency code
-		public String currencyCode;
-		// Original Json representation
-		public String originalJson;
-	}
+	private final GameActivity MainActivity;
 
-	public interface PurchaseLaunchCallback
-	{
-		void launchForResult(PendingIntent pendingIntent, int requestCode);
-	}
+	// Custom error value. Should match with EGooglePlayBillingResponseCode::CustomLogicError
+	private final int CustomLogicErrorResponse = -127;
 
-	public GooglePlayStoreHelper(String InProductKey, GameActivity InGameActivity, final Logger InLog)
+	public GooglePlayStoreHelper(GameActivity InGameActivity, Logger InLog)
 	{
-		// IAP is not ready to use until the service is instantiated.
-		bIsIapSetup = false;
-
 		Log = InLog;
 		Log.debug("[GooglePlayStoreHelper] - GooglePlayStoreHelper::GooglePlayStoreHelper");
 
-		gameActivity = InGameActivity;
-		productKey = InProductKey;
+		MainActivity = InGameActivity;
 
-		PurchasesUpdatedListener purchasesUpdatedListener = new PurchasesUpdatedListener(){
-			@Override
-			public void onPurchasesUpdated(BillingResult billingResult, List<Purchase> purchases){
-				if (billingResult.getResponseCode() == BillingClient.BillingResponseCode.OK
-					&& purchases != null) {
-						for (Purchase purchase : purchases) {
-							onPurchaseResult(billingResult, purchase);
-						}
-				} else
-				{ 
-					Log.debug("[GooglePlayStoreHelper] - GooglePlayStoreHelper::UserCancelled Purchase " + billingResult.getDebugMessage());
-					nativePurchaseComplete(billingResult.getResponseCode(), "", "", "", "");
-					// Handle an error caused by a user cancelling the purchase flow.
-				}
-			}
-		};
+		Client = BillingClient.newBuilder(MainActivity)
+			.setListener(this)
+			.enablePendingPurchases()
+			.build();
 
-		mBillingClient = BillingClient.newBuilder(gameActivity).setListener(purchasesUpdatedListener).enablePendingPurchases().build();
-		
-		mBillingClient.startConnection(new BillingClientStateListener() {
+		Client.startConnection(new BillingClientStateListener() {
 			@Override
-			public void onBillingSetupFinished(BillingResult billingResult) {
-				if (billingResult.getResponseCode() == BillingClient.BillingResponseCode.OK) {
+			public void onBillingSetupFinished(@NonNull BillingResult billingResult) {
+				if (billingResult.getResponseCode() == BillingClient.BillingResponseCode.OK) 
+				{
 					bIsIapSetup = true;
-					Log.debug("In-app billing supported for " + gameActivity.getPackageName());
+					Log.debug("[GooglePlayStoreHelper] - In-app billing supported for " + MainActivity.getPackageName());
 				}
 				else
 				{
-					Log.debug("In-app billing NOT supported for " + gameActivity.getPackageName() + " error " + billingResult.getResponseCode());
+					Log.debug("[GooglePlayStoreHelper] - In-app billing NOT supported for " + MainActivity.getPackageName() + " error " + TranslateBillingResult(billingResult));
 				}
 			}
-    
+
 			@Override
 			public void onBillingServiceDisconnected() {
 				// Try to restart the connection on the next request to
@@ -126,6 +82,7 @@ public class GooglePlayStoreHelper implements StoreHelper
 	/**
 	 * Determine whether the store is ready for purchases.
 	 */
+	@Override
 	public boolean IsAllowedToMakePurchases()
 	{
 		Log.debug("[GooglePlayStoreHelper] - GooglePlayStoreHelper::IsAllowedToMakePurchases");
@@ -133,115 +90,47 @@ public class GooglePlayStoreHelper implements StoreHelper
 	}
 
 	/**
-	 * Query product details for the provided skus
+	 * Query product details for the provided product ids
 	 */
+	@Override
 	public boolean QueryInAppPurchases(String[] InProductIDs)
 	{
 		Log.debug("[GooglePlayStoreHelper] - GooglePlayStoreHelper::QueryInAppPurchases");
+		Log.debug("[GooglePlayStoreHelper] - NumProducts: " + InProductIDs.length);
+		Log.debug("[GooglePlayStoreHelper] - Products: " + Arrays.toString(InProductIDs));
 
 		if (InProductIDs.length > 0)
 		{
-			ArrayList<String> skuList = new ArrayList<String>(InProductIDs.length);
-
-			for (String productId : InProductIDs)
-			{
-				Log.debug("[GooglePlayStoreHelper] - GooglePlayStoreHelper::QueryInAppPurchases - Querying " + productId);
-				skuList.add(productId);
-			}
-
-			Log.debug("[GooglePlayStoreHelper] - NumSkus: " + skuList.size());
-
-			SkuDetailsParams.Builder skuDetailsParams = SkuDetailsParams.newBuilder();
-			skuDetailsParams = skuDetailsParams.setSkusList(skuList).setType(BillingClient.SkuType.INAPP);
+			QueryProductDetailsParams Params = CreateQueryProductDetailsParamsForProducts(InProductIDs);
 			
-			mBillingClient.querySkuDetailsAsync(skuDetailsParams.build(),
-				new SkuDetailsResponseListener() {
-					@Override
-					public void onSkuDetailsResponse(BillingResult billingResult,
-						List<SkuDetails> skuDetailsList) 
-					{
-						// Process the result.
-						int response = billingResult.getResponseCode();
-
-						if(response == BillingClient.BillingResponseCode.OK && skuDetailsList != null)
-						{
-							Log.debug("[GooglePlayStoreHelper] - GooglePlayStoreHelper::QueryInAppPurchases - Response " + response + " SkuDetails:" + skuDetailsList.toString());
-							
-							ArrayList<GooglePlayProductDescription> productDescriptions = new ArrayList<GooglePlayProductDescription>();
-							for (SkuDetails thisSku : skuDetailsList)
-							{
-								GooglePlayProductDescription newDescription = new GooglePlayProductDescription();
-								newDescription.id = thisSku.getSku();
-								newDescription.title = thisSku.getTitle();
-								newDescription.description = thisSku.getDescription();
-								newDescription.price = thisSku.getPrice();
-								double priceRaw = thisSku.getPriceAmountMicros() / 1000000.0;
-								newDescription.priceRaw = (float)priceRaw;
-								newDescription.currencyCode = thisSku.getPriceCurrencyCode();
-								newDescription.originalJson = thisSku.getOriginalJson();
-								productDescriptions.add(newDescription);
-							}
-
-							// Should we send JSON, or somehow serialize this so we can have random access fields?
-							ArrayList<String> productIds = new ArrayList<String>();
-							ArrayList<String> titles = new ArrayList<String>();
-							ArrayList<String> descriptions = new ArrayList<String>();
-							ArrayList<String> prices = new ArrayList<String>();
-							ArrayList<Float> pricesRaw = new ArrayList<Float>();
-							ArrayList<String> currencyCodes = new ArrayList<String>();
-							ArrayList<String> originalJson = new ArrayList<String>();
-						
-							for (GooglePlayProductDescription product : productDescriptions)
-							{
-								productIds.add(product.id);
-								Log.debug("[GooglePlayStoreHelper] - GooglePlayStoreHelper::QueryInAppPurchases - Parsing details for: " + product.id);
-
-								titles.add(product.title);
-								Log.debug("[GooglePlayStoreHelper] - title: " + product.title);
-
-								descriptions.add(product.description);
-								Log.debug("[GooglePlayStoreHelper] - description: " + product.description);
-
-								prices.add(product.price);
-								Log.debug("[GooglePlayStoreHelper] - price: " + product.price);
-
-								pricesRaw.add(product.priceRaw);
-								Log.debug("[GooglePlayStoreHelper] - price_amount_micros: " + product.priceRaw);
-
-								currencyCodes.add(product.currencyCode);
-								Log.debug("[GooglePlayStoreHelper] - price_currency_code: " + product.currencyCode);
-
-								originalJson.add(product.originalJson);
-								Log.debug("[GooglePlayStoreHelper] - original_json: " + product.originalJson);
-
-							}
+			Client.queryProductDetailsAsync(Params, new ProductDetailsResponseListener() {
+				@Override
+				public void onProductDetailsResponse(@NonNull BillingResult Result, @NonNull List<ProductDetails> ProductDetailsList)
+				{
+					int ResponseCode = Result.getResponseCode();
 
-							float[] pricesRawPrimitive = new float[pricesRaw.size()];
-							for (int i = 0; i < pricesRaw.size(); i++)
-							{
-								pricesRawPrimitive[i] = pricesRaw.get(i);
-							}
+					if(ResponseCode == BillingClient.BillingResponseCode.OK)
+					{
+						Log.debug("[GooglePlayStoreHelper] - GooglePlayStoreHelper::QueryInAppPurchases - ProductDetails:" + ProductDetailsList.toString());
+						Log.debug("[GooglePlayStoreHelper] - GooglePlayStoreHelper::QueryInAppPurchases " + ProductDetailsList.size() + " items - Success!");
 
-							Log.debug("[GooglePlayStoreHelper] - GooglePlayStoreHelper::QueryInAppPurchases " + productIds.size() + " items - Success!");
-							nativeQueryComplete(BillingClient.BillingResponseCode.OK, productIds.toArray(new String[productIds.size()]), titles.toArray(new String[titles.size()]), descriptions.toArray(new String[descriptions.size()]), prices.toArray(new String[prices.size()]), pricesRawPrimitive, currencyCodes.toArray(new String[currencyCodes.size()]), originalJson.toArray(new String[originalJson.size()]));
-						}
-						else
-						{
-							Log.debug("[GooglePlayStoreHelper] - GooglePlayStoreHelper::QueryInAppPurchases - Failed with: " + response);
-							// If anything fails right now, stop immediately, not sure how to reconcile partial success/failure
-							nativeQueryComplete(response, null, null, null, null, null, null, null);
-						}
-						
-						Log.debug("[GooglePlayStoreHelper] - nativeQueryComplete done!");
+						ListOfProductDetailsAsArrays DetailArrays = new ListOfProductDetailsAsArrays(ProductDetailsList);
+						NativeQueryComplete(BillingClient.BillingResponseCode.OK, DetailArrays.ProductIds, DetailArrays.Titles, DetailArrays.Descriptions, DetailArrays.Prices, DetailArrays.PricesRaw, DetailArrays.CurrencyCodes);
+					}
+					else
+					{
+						Log.debug("[GooglePlayStoreHelper] - GooglePlayStoreHelper::QueryInAppPurchases - Failed with: " + TranslateBillingResult(Result));
+						NativeQueryComplete(ResponseCode, null, null, null, null, null, null);
 					}
+
+					Log.debug("[GooglePlayStoreHelper] - NativeQueryComplete done!");
 				}
-			);			
+			});
 		}
 		else
 		{
-			// nothing to query
 			Log.debug("[GooglePlayStoreHelper] - no products given to query");
-			nativeQueryComplete(UndefinedFailureResponse, null, null, null, null, null, null, null);
+			NativeQueryComplete(CustomLogicErrorResponse, null, null, null, null, null, null);
 			return false;
 		}
 
@@ -249,435 +138,329 @@ public class GooglePlayStoreHelper implements StoreHelper
 	}
 
 	/**
-	 * Start the purchase flow for a particular sku
+	 * Start the purchase flow for a particular product id
 	 */
-	final int purchaseIntentIdentifier = 1001;
-	public boolean BeginPurchase(String ProductID, String ObfuscatedAccountId)
+	@Override
+	public boolean BeginPurchase(final String ProductId, final String ObfuscatedAccountId)
 	{
-		ArrayList<String> skuList = new ArrayList<String>(1);
-		skuList.add(ProductID);
+		QueryProductDetailsParams Params = CreateQueryProductDetailsParamsForProduct(ProductId);
 
-		SkuDetailsParams.Builder skuDetailsParams = SkuDetailsParams.newBuilder();
-		skuDetailsParams.setSkusList(skuList).setType(BillingClient.SkuType.INAPP);
-		
-		class SkuDetailsResponseListenerImpl implements SkuDetailsResponseListener
-		{
-			String ObfuscatedAccountId = null;
-			public SkuDetailsResponseListenerImpl setObfuscatedAccountId(String InObfuscatedAccountId)
-			{
-				ObfuscatedAccountId = InObfuscatedAccountId;
-				return this;
-			}
-			
+		Client.queryProductDetailsAsync(Params, new ProductDetailsResponseListener() {
 			@Override
-			public void onSkuDetailsResponse(BillingResult billingResult,
-				List<SkuDetails> skuDetailsList) 
+			public void onProductDetailsResponse(@NonNull BillingResult Result, @NonNull List<ProductDetails> ProductDetailsList)
 			{
-				int response = billingResult.getResponseCode();
-				if(response == BillingClient.BillingResponseCode.OK && skuDetailsList != null)
+				if(Result.getResponseCode() == BillingClient.BillingResponseCode.OK)
 				{
-					for(SkuDetails skuDetails : skuDetailsList)
+					final ProductDetails DetailsForProductId = FindProductDetailsInList(ProductId, ProductDetailsList);
+					if(DetailsForProductId != null)
 					{
-						BillingFlowParams.Builder flowParams = BillingFlowParams.newBuilder().setSkuDetails(skuDetails);
-						
-						if(ObfuscatedAccountId != null)
-						{
-							flowParams = flowParams.setObfuscatedAccountId(ObfuscatedAccountId);
-						}
-	
-						if (gameActivity.IsInVRMode())
+						// BillingClient.launchBillingFlow is annotated @UiThread, so dispatch the call to the Activity UI thread
+						MainActivity.runOnUiThread(new Runnable()
 						{
-							int responseCode = mBillingClient.isFeatureSupported(BillingClient.FeatureType.IN_APP_ITEMS_ON_VR).getResponseCode();
-							boolean isSupportedInVR = (responseCode == BillingClient.BillingResponseCode.OK);
-							if (isSupportedInVR)
-							{
-								Log.debug("[GooglePlayStoreHelper] - GooglePlayStoreHelper::BeginPurchase - v7 VR purchase" + skuDetails.getSku());
-								flowParams = flowParams.setVrPurchaseFlow(true);
-							}
-							else
-							{
-								Log.debug("[GooglePlayStoreHelper] - GooglePlayStoreHelper::BeginPurchase - v3 IAB purchase:" + skuDetails.getSku());
-								flowParams = flowParams.setVrPurchaseFlow(false);
+							@Override
+							public void run() {
+								Log.debug("[GooglePlayStoreHelper] - GooglePlayStoreHelper::BeginPurchase - Launching billing flow " + ProductId);
+								BillingFlowParams Params = CreateBillingFlowParams(DetailsForProductId, ObfuscatedAccountId);
+								BillingResult Result = Client.launchBillingFlow(MainActivity, Params);
+								int ResponseCode = Result.getResponseCode();
+								if (ResponseCode != BillingClient.BillingResponseCode.OK) 
+								{
+									Log.debug("[GooglePlayStoreHelper] - GooglePlayStoreHelper::BeginPurchase - Failed! " + TranslateBillingResult(Result));
+									NativePurchaseComplete(ResponseCode, ProductId, Purchase.PurchaseState.UNSPECIFIED_STATE, "", "", "");
+								}
 							}
-						}
-						else
-						{
-							Log.debug("[GooglePlayStoreHelper] - GooglePlayStoreHelper::BeginPurchase - v3 IAB purchase:" + skuDetails.getSku());
-							flowParams = flowParams.setVrPurchaseFlow(false);
-						}
-
-						int responseCode = mBillingClient.launchBillingFlow(gameActivity, flowParams.build()).getResponseCode();
-						Log.debug("[GooglePlayStoreHelper] - GooglePlayStoreHelper::BeginPurchase - Launching billing flow " + skuDetails.getSku());
+						});
+					}
+					else
+					{
+						Log.debug("[GooglePlayStoreHelper] - GooglePlayStoreHelper::BeginPurchase - Failed! Could not find ProductDetails for " + ProductId);
+						NativePurchaseComplete(CustomLogicErrorResponse, ProductId, Purchase.PurchaseState.UNSPECIFIED_STATE,"", "", "");
 					}
 				}
 				else
 				{
-					Log.debug("[GooglePlayStoreHelper] - GooglePlayStoreHelper::BeginPurchase - Failed! " + TranslateServerResponseCode(response));
-					nativePurchaseComplete(UndefinedFailureResponse, "", "", "", "");
+					Log.debug("[GooglePlayStoreHelper] - GooglePlayStoreHelper::BeginPurchase - Failed! " + TranslateBillingResult(Result));
+					NativePurchaseComplete(Result.getResponseCode(), ProductId, Purchase.PurchaseState.UNSPECIFIED_STATE,"", "", "");
 				}
 			}
-		}
-		
-		mBillingClient.querySkuDetailsAsync(skuDetailsParams.build(), new SkuDetailsResponseListenerImpl().setObfuscatedAccountId(ObfuscatedAccountId));
+		});
 		return true;
 	}
 
-	/**
-	 * Restore previous purchases the user may have made.
-	 */
-	public boolean RestorePurchases(String[] InProductIDs, boolean[] bConsumable)
+	@Override
+	public boolean QueryExistingPurchases()
 	{
-		Log.debug("[GooglePlayStoreHelper] - GooglePlayStoreHelper::RestorePurchases");
-
-		ArrayList<Purchase> ownedProducts = new ArrayList<Purchase>();
-		int responseCode = GatherOwnedPurchaseData(ownedProducts);
-		if (responseCode == BillingClient.BillingResponseCode.OK)
-		{
-			Log.debug("[GooglePlayStoreHelper] - GooglePlayStoreHelper::RestorePurchases - User has previously purchased " + ownedProducts.size() + " inapp products" );
-
-			final ArrayList<Purchase> f_ownedProducts = ownedProducts;
-
-			final ArrayList<String> f_ownedSkus = new ArrayList<String>();
-			final ArrayList<String> f_purchaseDataList = new ArrayList<String>();
-			final ArrayList<String> f_signatureList = new ArrayList<String>();
-			final String[] RestoreProductIDs = InProductIDs;
-			final boolean[] bRestoreConsumableFlags = bConsumable;
+		Log.debug("[GooglePlayStoreHelper] - GooglePlayStoreHelper::QueryExistingPurchases");
 
-			// Is this here to ensure callbacks are fired on the main thread?
-			Handler mainHandler = new Handler(Looper.getMainLooper());
-			mainHandler.post(new Runnable()
+		QueryPurchasesParams Params =  QueryPurchasesParams.newBuilder()
+			.setProductType(BillingClient.ProductType.INAPP)
+			.build();
+		
+		Client.queryPurchasesAsync(Params, new PurchasesResponseListener() {
+			@Override
+			public void onQueryPurchasesResponse(@NonNull BillingResult Result, @NonNull List<Purchase> OwnedItems) 
 			{
-				@Override
-				public void run()
+				int ResponseCode = Result.getResponseCode();
+				if (ResponseCode == BillingClient.BillingResponseCode.OK) 
 				{
-					Log.debug("[GooglePlayStoreHelper] - GooglePlayStoreHelper::RestorePurchases - Now consuming any purchases that may have been missed.");
-					final ArrayList<String> productTokens = new ArrayList<String>();
-					final ArrayList<String> receipts = new ArrayList<String>();
-					final ArrayList<Integer> cachedErrorResponses = new ArrayList<Integer>();
-
-					bIsConsumeComplete = false;
-					// @todo possible bug, restores ALL ownedSkus, not just those provided by caller in RestoreProductIDs
-					for (Purchase ownedProduct : f_ownedProducts)
-					{
-						
-						productTokens.add(ownedProduct.getPurchaseToken());
-
-						boolean bTryToConsume = false;
-
-						// This is assuming that all purchases should be consumed. Consuming a purchase that is meant to be a one-time purchase makes it so the
-						// user is able to buy it again. Also, it makes it so the purchase will not be able to be restored again in the future.
-
-						for (int Idy = 0; Idy < RestoreProductIDs.length; Idy++)
-						{
-							if( ownedProduct.getSku().equals(RestoreProductIDs[Idy]) )
-							{
-								if(Idy < bRestoreConsumableFlags.length)
-								{
-									bTryToConsume = bRestoreConsumableFlags[Idy];
-									Log.debug("[GooglePlayStoreHelper] - GooglePlayStoreHelper::RestorePurchases - Found Consumable Flag for Product " + ownedProduct.getSku() + " bConsumable = " + bTryToConsume);
-								}
-								break;
-							}
-						}
+					Log.debug("[GooglePlayStoreHelper] - GooglePlayStoreHelper::QueryExistingPurchases - User has previously purchased " + OwnedItems.size() + " inapp products");
+					Log.debug("[GooglePlayStoreHelper] - GooglePlayStoreHelper::QueryExistingPurchases - Success!");
 
-						final Purchase constProduct = ownedProduct;
-						if (bTryToConsume)
-						{
-							Log.debug("[GooglePlayStoreHelper] - GooglePlayStoreHelper::RestorePurchases - Attempting to consume " + ownedProduct.getSku());
-							ConsumeParams consumeParams = ConsumeParams.newBuilder().setPurchaseToken(ownedProduct.getPurchaseToken()).build();
-							mBillingClient.consumeAsync(consumeParams, (new ConsumeResponseListener() {
-								@Override
-								public void onConsumeResponse(BillingResult billingResult, String purchaseToken) {
-									if(billingResult.getResponseCode() == BillingClient.BillingResponseCode.OK)
-									{
-										// How do we get purchase here.. need to figure out a way to persist cachedResponse and receipts
-										Log.debug("[GooglePlayStoreHelper] - GooglePlayStoreHelper::RestorePurchases - Purchase restored for " + constProduct.getSku());
-										String receipt = Base64.encode(constProduct.getOriginalJson().getBytes());
-										receipts.add(receipt);
-										
-										f_ownedSkus.add(constProduct.getSku());
-										f_signatureList.add(constProduct.getSignature());
-									}
-									else
-									{
-										Log.debug("[GooglePlayStoreHelper] - GooglePlayStoreHelper::RestorePurchases - consumePurchase failed for " + constProduct.getSku() + " with error " + billingResult.getResponseCode());
-										receipts.add("");
-										cachedErrorResponses.add(billingResult.getResponseCode());
-									}
-									outstandingConsumeRequests--;
-
-									if(outstandingConsumeRequests <= 0)
-									{
-										Log.debug("[GooglePlayStoreHelper] - GooglePlayStoreHelper::RestorePurchases - Success!");
-										int finalResponse = (cachedErrorResponses.size() > 0) ? cachedErrorResponses.get(0) : BillingClient.BillingResponseCode.OK;
-										nativeRestorePurchasesComplete(finalResponse, f_ownedSkus.toArray(new String[f_ownedSkus.size()]), productTokens.toArray(new String[productTokens.size()]), receipts.toArray(new String[receipts.size()]), f_signatureList.toArray(new String[f_signatureList.size()]));
-										bIsConsumeComplete = true;
-									}
-								}
-							}));
-								
-							outstandingConsumeRequests++;
-						}
-						else
-						{
-							// How do we get purchase here.. need to figure out a way to persist cachedResponse and receipts
-							Log.debug("[GooglePlayStoreHelper] - GooglePlayStoreHelper::RestorePurchases - Purchase restored for " + constProduct.getSku());
-							String receipt = Base64.encode(constProduct.getOriginalJson().getBytes());
-							receipts.add(receipt);
-										
-							f_ownedSkus.add(constProduct.getSku());
-							f_signatureList.add(constProduct.getSignature());
-						}
-					}
-
-					// There was nothing to consume, so we are done!
-					if(outstandingConsumeRequests <= 0 && !bIsConsumeComplete)
-					{
-						Log.debug("[GooglePlayStoreHelper] - GooglePlayStoreHelper::RestorePurchases - Success!");
-						nativeRestorePurchasesComplete(BillingClient.BillingResponseCode.OK, f_ownedSkus.toArray(new String[f_ownedSkus.size()]), productTokens.toArray(new String[productTokens.size()]), receipts.toArray(new String[receipts.size()]), f_signatureList.toArray(new String[f_signatureList.size()]));
-					}
+					ListOfPurchasesAsArrays PurchasesArrays = new ListOfPurchasesAsArrays(OwnedItems);
+					NativeQueryExistingPurchasesComplete(ResponseCode, PurchasesArrays.OwnedProducts, PurchasesArrays.PurchaseStates, PurchasesArrays.ProductTokens, PurchasesArrays.Receipts, PurchasesArrays.Signatures);
+				} 
+				else 
+				{
+					Log.debug("[GooglePlayStoreHelper] - GooglePlayStoreHelper::QueryExistingPurchases - Failed to collect existing purchases" + TranslateBillingResult(Result));
+					NativeQueryExistingPurchasesComplete(ResponseCode, null, null, null, null, null);
 				}
-			});
-		}
-		else
-		{
-			nativeRestorePurchasesComplete(responseCode, null, null, null, null);
-			Log.debug("[GooglePlayStoreHelper] - GooglePlayStoreHelper::RestorePurchases - Failed to collect existing purchases");
-			return false;
-		}
-
+			}
+		});
 		return true;
 	}
 
-	public boolean QueryExistingPurchases()
+	@Override
+	public void ConsumePurchase(final String PurchaseToken)
 	{
-		Log.debug("[GooglePlayStoreHelper] - GooglePlayStoreHelper::QueryExistingPurchases");
-		ArrayList<String> ownedSkus = new ArrayList<String>();
-		ArrayList<String> purchaseDataList = new ArrayList<String>();
-		ArrayList<String> signatureList = new ArrayList<String>();
-
-		ArrayList<Purchase> ownedProducts = new ArrayList<Purchase>();
+		Log.debug("[GooglePlayStoreHelper] - Beginning ConsumePurchase: " + PurchaseToken);
 
-		int responseCode = GatherOwnedPurchaseData(ownedProducts);
-		if (responseCode == BillingClient.BillingResponseCode.OK)
-		{
-			Log.debug("[GooglePlayStoreHelper] - GooglePlayStoreHelper::QueryExistingPurchases - User has previously purchased " + ownedSkus.size() + " inapp products" );
+		ConsumeParams Params = ConsumeParams.newBuilder()
+			.setPurchaseToken(PurchaseToken)
+			.build();
+		Client.consumeAsync(Params, new ConsumeResponseListener() {
+			@Override
+			public void onConsumeResponse(@NonNull BillingResult Result, @NonNull String PurchaseToken) 
+			{
+				Log.debug("[GooglePlayStoreHelper] - ConsumePurchase response: " + TranslateBillingResult(Result));
+				NativeConsumeComplete(Result.getResponseCode(), PurchaseToken);
+			}
+		});
+	}
 
-			ArrayList<String> productTokens = new ArrayList<String>();
-			ArrayList<String> receipts = new ArrayList<String>();
+	///////////////////////////////////////////////////////
+	// The Google's PurchasesUpdatedListener interface implementation.
 
-			for (Purchase purchase : ownedProducts)
+	@Override
+	public void onPurchasesUpdated(@NonNull BillingResult Result, @Nullable List<Purchase> UpdatedPurchases){
+		int responseCode = Result.getResponseCode();
+		Log.debug("[GooglePlayStoreHelper] - GooglePlayStoreHelper::onPurchasesUpdated - Processing purchase result. Response Code: " + TranslateBillingResult(Result));
+		if (responseCode == BillingClient.BillingResponseCode.OK && UpdatedPurchases != null) 
+		{
+			for (final Purchase UpdatedPurchase : UpdatedPurchases) 
 			{
-				productTokens.add(purchase.getPurchaseToken());
-				ownedSkus.add(purchase.getSku());
-				signatureList.add(purchase.getSignature());
-				String receipt = Base64.encode(purchase.getOriginalJson().getBytes());
-				receipts.add(receipt);
-			}
+				final String ProductId = UpdatedPurchase.getProducts().get(0);
+				Log.debug("[GooglePlayStoreHelper] - GooglePlayStoreHelper::onPurchasesUpdated - Processing purchase of product " + ProductId);
+				Log.debug("[GooglePlayStoreHelper] - Purchase data: " + UpdatedPurchase.toString());
+				Log.debug("[GooglePlayStoreHelper] - Data signature: " + UpdatedPurchase.getSignature());
 
-			Log.debug("[GooglePlayStoreHelper] - GooglePlayStoreHelper::QueryExistingPurchases - Success!");
-			nativeQueryExistingPurchasesComplete(responseCode, ownedSkus.toArray(new String[ownedSkus.size()]), productTokens.toArray(new String[productTokens.size()]), receipts.toArray(new String[receipts.size()]), signatureList.toArray(new String[signatureList.size()]));
-			return true;
+				String Receipt = Base64.encode(UpdatedPurchase.getOriginalJson().getBytes());
+				NativePurchaseComplete(BillingClient.BillingResponseCode.OK, ProductId, UpdatedPurchase.getPurchaseState(), UpdatedPurchase.getPurchaseToken(), Receipt, UpdatedPurchase.getSignature());
+			}
 		}
 		else
 		{
-			Log.debug("[GooglePlayStoreHelper] - GooglePlayStoreHelper::QueryExistingPurchases - Failed to collect existing purchases" +  TranslateServerResponseCode(responseCode));
-			nativeQueryExistingPurchasesComplete(responseCode, null, null, null, null);
-			return false;
+			Log.debug("[GooglePlayStoreHelper] - GooglePlayStoreHelper::onPurchasesUpdated - Purchase failed. Response code: " + TranslateBillingResult(Result));
+ 			NativePurchaseComplete(responseCode, "", Purchase.PurchaseState.UNSPECIFIED_STATE, "", "", "");
 		}
 	}
 
-	public void ConsumePurchase(String purchaseToken)
-	{
-		Log.debug("[GooglePlayStoreHelper] - Beginning ConsumePurchase: " + purchaseToken);
-
-		final String f_purchaseToken = purchaseToken;
+	///////////////////////////////////////////////////////
+	// Some helpers
 
-		Handler mainHandler = new Handler(Looper.getMainLooper());
-		mainHandler.post(new Runnable()
+	class ListOfProductDetailsAsArrays
+	{
+		public String[] ProductIds;
+		public String[] Titles;
+		public String[] Descriptions;
+		public String[] Prices;
+		public float[] PricesRaw;
+		public String[] CurrencyCodes;
+
+		public ListOfProductDetailsAsArrays(List<ProductDetails> InProductDetailsList)
 		{
-			@Override
-			public void run()
+			ProductIds = new String[InProductDetailsList.size()];
+			Titles = new String[InProductDetailsList.size()];
+			Descriptions = new String[InProductDetailsList.size()];
+			Prices = new String[InProductDetailsList.size()];
+			PricesRaw = new float[InProductDetailsList.size()];
+			CurrencyCodes = new String[InProductDetailsList.size()];
+
+			for (int index = 0; index < InProductDetailsList.size(); index++) 
 			{
-				ConsumeParams consumeParams = ConsumeParams.newBuilder().setPurchaseToken(f_purchaseToken).build();
-				mBillingClient.consumeAsync(consumeParams, (new ConsumeResponseListener() {
-					@Override
-					public void onConsumeResponse(BillingResult billingResult, String purchaseToken) {
-						if(billingResult.getResponseCode() == BillingClient.BillingResponseCode.OK)
-						{
-							Log.debug("[GooglePlayStoreHelper] - ConsumePurchase success");
-						}
-						else
-						{
-							Log.debug("[GooglePlayStoreHelper] - ConsumePurchase failed with error " + TranslateServerResponseCode(billingResult.getResponseCode()));
-						}
-					}
-				}));
+				ProductDetails ThisProduct = InProductDetailsList.get(index);
+				ProductDetails.OneTimePurchaseOfferDetails PriceDetails = ThisProduct.getOneTimePurchaseOfferDetails();
+
+				Log.debug("[GooglePlayStoreHelper] - GooglePlayStoreHelper::QueryInAppPurchases - Parsing details for: " + ThisProduct.getProductId());
+				Log.debug("[GooglePlayStoreHelper] - title: " + ThisProduct.getTitle());
+				Log.debug("[GooglePlayStoreHelper] - description: " + ThisProduct.getDescription());
+				Log.debug("[GooglePlayStoreHelper] - price: " + PriceDetails.getFormattedPrice());
+				Log.debug("[GooglePlayStoreHelper] - price_amount_micros: " + PriceDetails.getPriceAmountMicros());
+				Log.debug("[GooglePlayStoreHelper] - price_currency_code: " + PriceDetails.getPriceCurrencyCode());
+
+				ProductIds[index] = ThisProduct.getProductId();
+				Titles[index] = ThisProduct.getTitle();
+				Descriptions[index] = ThisProduct.getDescription();
+				Prices[index] = PriceDetails.getFormattedPrice();
+				PricesRaw[index] = (float)(PriceDetails.getPriceAmountMicros() / 1000000.0);
+				CurrencyCodes[index] = PriceDetails.getPriceCurrencyCode();
 			}
-		});
+		}
 	}
-	
-	public boolean onActivityResult(int requestCode, int resultCode, Intent data)
+
+	class ListOfPurchasesAsArrays 
 	{
-		Log.debug("[GooglePlayStoreHelper] - GooglePlayStoreHelper::onActivityResult unimplemented on BillingApiV2");
-		return false;
-	}
+		public String[] OwnedProducts;
+		public int[] PurchaseStates;
+		public String[] ProductTokens;
+		public String[] Signatures;
+		public String[] Receipts;
 
-	///////////////////////////////////////////////////////
-	// Game Activity/Context driven methods we need to listen for.
+		public ListOfPurchasesAsArrays(List<Purchase> InOwnedItems) 
+		{
+			OwnedProducts = new String[InOwnedItems.size()];
+			PurchaseStates = new int[InOwnedItems.size()];
+			ProductTokens = new String[InOwnedItems.size()];
+			Signatures = new String[InOwnedItems.size()];
+			Receipts = new String[InOwnedItems.size()];
 
-	/**
-	 * On Destory we should unbind our IInAppBillingService service
-	 */
-	public void onDestroy()
+			for (int index = 0; index < InOwnedItems.size(); index ++) 
+			{
+				Purchase OwnedPurchase = InOwnedItems.get(index);
+				OwnedProducts[index] = OwnedPurchase.getProducts().get(0);
+				PurchaseStates[index] = OwnedPurchase.getPurchaseState();
+				ProductTokens[index] = OwnedPurchase.getPurchaseToken();
+				Signatures[index] = OwnedPurchase.getSignature();
+				Receipts[index] = Base64.encode(OwnedPurchase.getOriginalJson().getBytes());
+			}
+		}
+	}
+	
+	private static BillingFlowParams CreateBillingFlowParams(ProductDetails InProductDetails, String InObfuscatedAccountId)
 	{
-		Log.debug("[GooglePlayStoreHelper] - GooglePlayStoreHelper::onDestroy");
-
-		if (mBillingClient != null)
+		ArrayList<BillingFlowParams.ProductDetailsParams> ProductDetailsParamList = new ArrayList<>(1);
+		ProductDetailsParamList.add(BillingFlowParams.ProductDetailsParams.newBuilder()
+			.setProductDetails(InProductDetails)
+			.build());
+		
+		BillingFlowParams.Builder Params = BillingFlowParams.newBuilder()
+			.setProductDetailsParamsList(ProductDetailsParamList);
+		if (InObfuscatedAccountId != null)
 		{
-			mBillingClient = null; //WMM ???
+			Params = Params.setObfuscatedAccountId(InObfuscatedAccountId);
 		}
+		return Params.build();
 	}
 
-	/**
-	 * Route taken by the Purchase workflow. We listen for our purchaseIntentIdentifier request code and
-	 * handle the response accordingly
-	 */
-	public boolean onPurchaseResult(BillingResult billingResult, Purchase purchase)
+	private static QueryProductDetailsParams CreateQueryProductDetailsParamsForProduct(String InProductId)
 	{
-		Log.debug("[GooglePlayStoreHelper] - GooglePlayStoreHelper::onPurchaseResult");
+		ArrayList<QueryProductDetailsParams.Product> ProductList = new ArrayList<>(1);
+		ProductList.add(QueryProductDetailsParams.Product.newBuilder()
+			.setProductId(InProductId)
+			.setProductType(BillingClient.ProductType.INAPP)
+			.build());
+		return QueryProductDetailsParams.newBuilder()
+			.setProductList(ProductList)
+			.build();
+	}
 
-		if (billingResult == null)
+	private static QueryProductDetailsParams CreateQueryProductDetailsParamsForProducts(String[] InProductIds)
+	{
+		ArrayList<QueryProductDetailsParams.Product> ProductList = new ArrayList<>();
+		for (String ProductId : InProductIds)
 		{
-			Log.debug("Null data in purchase activity result.");
-			nativePurchaseComplete(UndefinedFailureResponse, "", "", "", "");
-			return true;
+			QueryProductDetailsParams.Product Product = QueryProductDetailsParams.Product.newBuilder()
+				.setProductId(ProductId)
+				.setProductType(BillingClient.ProductType.INAPP)
+				.build();
+			ProductList.add(Product);
 		}
 
-		int responseCode = billingResult.getResponseCode();
-		if (responseCode == BillingClient.BillingResponseCode.OK)
-		{
-			Log.debug("[GooglePlayStoreHelper] - GooglePlayStoreHelper::onActivityResult - Processing purchase result. Response Code: " + TranslateServerResponseCode(responseCode));
-			Log.debug("Purchase data: " + purchase.toString());
-			Log.debug("Data signature: " + purchase.getSignature());
+		return QueryProductDetailsParams.newBuilder()
+			.setProductList(ProductList)
+			.build();
+	}
 
-			if(purchase.getPurchaseState() == Purchase.PurchaseState.PURCHASED)
-			{
-				final String sku = purchase.getSku();
-				final Purchase f_purchase = purchase;
-				Handler mainHandler = new Handler(Looper.getMainLooper());
-				mainHandler.post(new Runnable()
-				{
-					@Override
-					public void run()
-					{
-						String receipt = Base64.encode(f_purchase.getOriginalJson().getBytes());
-						nativePurchaseComplete(BillingClient.BillingResponseCode.OK, sku, f_purchase.getPurchaseToken(), receipt, f_purchase.getSignature());
-					}
-				});
-			}
-			else
+	private static ProductDetails FindProductDetailsInList(String ProductId, List<ProductDetails> InProductDetailsList)
+	{
+		for(ProductDetails Details : InProductDetailsList) 
+		{
+			if (Details.getProductId().equals(ProductId)) 
 			{
-				// This is a pending purchase.. need to hold off on reporting as purchase to game until it actually occurs.
-				// Special case - if pending purchase is USER_CANCELLED but the SKU is returned, then it is deferred
-				final String sku = purchase.getSku();
-				nativePurchaseComplete(BillingClient.BillingResponseCode.USER_CANCELED, sku, "", "", "");
+				return Details;
 			}
 		}
-		else if (responseCode == BillingClient.BillingResponseCode.USER_CANCELED)
-		{
-			Log.debug("[GooglePlayStoreHelper] - GooglePlayStoreHelper::onActivityResult - Purchase canceled." + TranslateServerResponseCode(responseCode));
-			nativePurchaseComplete(BillingClient.BillingResponseCode.USER_CANCELED, "", "", "", "");
-		}
-		else
-		{
-			Log.debug("[GooglePlayStoreHelper] - GooglePlayStoreHelper::onActivityResult - Purchase failed. Result code: " +
-				Integer.toString(responseCode) + ". Response: " + TranslateServerResponseCode(responseCode));
-			nativePurchaseComplete(UndefinedFailureResponse, "", "", "", "");
-		}
-
-		return true;
+		return null;
 	}
-
-
-	///////////////////////////////////////////////////////
-	// Internal helper functions that deal assist with various IAB related events
-
+	
 	/**
-	 * Get a text tranlation of the Response Codes returned by google play.
+	 * Get a text translation of the Response Codes returned by google play.
 	 */
-	private String TranslateServerResponseCode(final int serverResponseCode)
+
+	private static String TranslateBillingResponseCode(int ResponseCode)
 	{
-		// taken from https://developer.android.com/reference/com/android/billingclient/api/BillingClient.BillingResponseCode
-		switch(serverResponseCode)
+		switch(ResponseCode)
 		{
 			case BillingClient.BillingResponseCode.OK:
-				return "Success";
+				return "BillingResponseCode.OK";
 			case BillingClient.BillingResponseCode.USER_CANCELED:
-				return "User pressed back or canceled a dialog";
+				return "BillingResponseCode.USER_CANCELED";
 			case BillingClient.BillingResponseCode.SERVICE_UNAVAILABLE:
-				return "Network connection is down";
+				return "BillingResponseCode.SERVICE_UNAVAILABLE";
 			case BillingClient.BillingResponseCode.SERVICE_TIMEOUT:
-				return "The request has reached the maximum timeout before Google Play responds";
+				return "BillingResponseCode.SERVICE_TIMEOUT";
 			case BillingClient.BillingResponseCode.SERVICE_DISCONNECTED:
-				return "Play Store service is not connected now - potentially transient state";
+				return "BillingResponseCode.SERVICE_DISCONNECTED";
 			case BillingClient.BillingResponseCode.BILLING_UNAVAILABLE:
-				return "Billing API version is not supported for the type requested";
+				return "BillingResponseCode.BILLING_UNAVAILABLE";
 			case BillingClient.BillingResponseCode.FEATURE_NOT_SUPPORTED:
-				return "Requested feature is not supported by Play Store on the current device";
+				return "BillingResponseCode.FEATURE_NOT_SUPPORTED";
 			case BillingClient.BillingResponseCode.ITEM_UNAVAILABLE:
-				return "Requested product is not available for purchase";
+				return "BillingResponseCode.ITEM_UNAVAILABLE";
 			case BillingClient.BillingResponseCode.DEVELOPER_ERROR:
-				return "Invalid arguments provided to the API. This error can also indicate that the application was not correctly signed or properly set up for In-app Billing in Google Play, or does not have the necessary permissions in its manifest";
+				return "BillingResponseCode.DEVELOPER_ERROR";
 			case BillingClient.BillingResponseCode.ERROR:
-				return "Fatal error during the API action";
+				return "BillingResponseCode.ERROR";
 			case BillingClient.BillingResponseCode.ITEM_ALREADY_OWNED:
-				return "Failure to purchase since item is already owned";
+				return "BillingResponseCode.ITEM_ALREADY_OWNED";
 			case BillingClient.BillingResponseCode.ITEM_NOT_OWNED:
-				return "Failure to consume since item is not owned";
+				return "BillingResponseCode.ITEM_NOT_OWNED";
 			default:
-				return "Unknown Server Response Code";
+				return "Unknown Response Code: " + ResponseCode;
 		}
 	}
 
+	private static String TranslateBillingResult(@NonNull BillingResult Result)
+	{
+		int ResponseCode = Result.getResponseCode();
+		switch(ResponseCode)
+		{
+			case BillingClient.BillingResponseCode.OK:
+			case BillingClient.BillingResponseCode.USER_CANCELED:
+				return TranslateBillingResponseCode(ResponseCode);
+			default:
+				return TranslateBillingResponseCode(ResponseCode) + " DebugMessage: " + Result.getDebugMessage();
+		}
+	}
+
+	///////////////////////////////////////////////////////
+	// Game Activity/Context driven methods we need to listen for.
+
 	/**
-	 * Recursive functionality to gather all of the purchases owned by a user.
-	 * if the user owns a lot of products then we need to getPurchases again with a continuationToken
+	 * On Destroy disconnect from service
 	 */
-	private int GatherOwnedPurchaseData(ArrayList<Purchase> inPurchases)
+	public void onDestroy()
 	{
-		int responseCode = UndefinedFailureResponse;
-		try
-		{
-			Purchase.PurchasesResult purchasesResult = mBillingClient.queryPurchases(BillingClient.SkuType.INAPP);
+		Log.debug("[GooglePlayStoreHelper] - GooglePlayStoreHelper::onDestroy");
 
-			responseCode = purchasesResult.getResponseCode();
-			Log.debug("[GooglePlayStoreHelper] - GooglePlayStoreHelper::GatherOwnedPurchaseData - getPurchases result. Response Code: " + TranslateServerResponseCode(responseCode));
-			if (responseCode == BillingClient.BillingResponseCode.OK)
-			{
-				List<Purchase> ownedItems = purchasesResult.getPurchasesList();
-				for (Purchase purchase : ownedItems)
-				{
-					// Don't Mark pending purchases as owned
-					if(purchase.getPurchaseState() == Purchase.PurchaseState.PURCHASED)
-					{
-						inPurchases.add(purchase);
-					}
-				}
-			}
-		}
-		catch (Exception e)
+		if (Client != null)
 		{
-			Log.debug("[GooglePlayStoreHelper] - GooglePlayStoreHelper::GatherOwnedPurchaseData - Failed for purchase request!. " + e.getMessage());
+			Client.endConnection();
 		}
-
-		return responseCode;
 	}
 
 	// Callback that notify the C++ implementation that a task has completed
-	public native void nativeQueryComplete(int responseCode, String[] productIDs, String[] titles, String[] descriptions, String[] prices, float[] pricesRaw, String[] currencyCodes, String[] originalJson );
-	public native void nativePurchaseComplete(int responseCode, String ProductId, String ProductToken, String ReceiptData, String Signature);
-	public native void nativeRestorePurchasesComplete(int responseCode, String[] ProductIds, String[] ProductTokens, String[] ReceiptsData, String[] Signatures);
-	public native void nativeQueryExistingPurchasesComplete(int responseCode, String[] ProductIds, String[] ProductTokens, String[] ReceiptsData, String[] Signatures);
+	public native void NativeQueryComplete(int ResponseCode, String[] ProductIDs, String[] Titles, String[] Descriptions, String[] Prices, float[] PricesRaw, String[] CurrencyCodes);
+	public native void NativePurchaseComplete(int ResponseCode, String ProductId, int PurchaseState, String ProductToken, String ReceiptData, String Signature);
+	public native void NativeQueryExistingPurchasesComplete(int ResponseCode, String[] ProductIds, int[] PurchaseState, String[] ProductTokens, String[] ReceiptsData, String[] Signatures);
+	public native void NativeConsumeComplete(int ResponseCode, String PurchaseToken);
 }
-
diff --git a/Engine/Plugins/Online/Android/OnlineSubsystemGooglePlay/Source/OnlineSubsystemGooglePlay_UPL.xml b/Engine/Plugins/Online/Android/OnlineSubsystemGooglePlay/Source/OnlineSubsystemGooglePlay_UPL.xml
--- a/Engine/Plugins/Online/Android/OnlineSubsystemGooglePlay/Source/OnlineSubsystemGooglePlay_UPL.xml
+++ b/Engine/Plugins/Online/Android/OnlineSubsystemGooglePlay/Source/OnlineSubsystemGooglePlay_UPL.xml
@@ -15,7 +15,7 @@
 	<buildGradleAdditions>
 		<insert>
 			dependencies {
-			implementation 'com.android.billingclient:billing:3.0.0'
+			implementation 'com.android.billingclient:billing:5.0.0'
 			}
 		</insert>
 	</buildGradleAdditions>
@@ -71,24 +71,11 @@ import java.security.NoSuchAlgorithmException;
 		</insert>
 	</gameActivityImportAdditions>
 
-	<!-- optional additions to GameActivity class in GameActivity.java -->
-	<gameActivityClassAdditions>
-		<insert>
-	/** Implement this if app wants to handle IAB activity launch. **/
-	GooglePlayStoreHelper.PurchaseLaunchCallback purchaseLaunchCallback = null;
-
-	public GooglePlayStoreHelper.PurchaseLaunchCallback getPurchaseLaunchCallback()
-	{
-		return purchaseLaunchCallback;
-	}
-		</insert>
-	</gameActivityClassAdditions>
-
 	<gameActivityIapSetupServiceAdditions>
 		<insert>
 		if (getPackageManager().checkPermission("com.android.vending.BILLING", getPackageName()) == getPackageManager().PERMISSION_GRANTED)
 		{
-			IapStoreHelper = new GooglePlayStoreHelper(InProductKey, this, Log);
+			IapStoreHelper = new GooglePlayStoreHelper(this, Log);
 			if (IapStoreHelper != null)
 			{
 				Log.debug("[JAVA] - AndroidThunkJava_IapSetupService - Setup started");
diff --git a/Engine/Plugins/Online/Android/OnlineSubsystemGooglePlay/Source/Private/OnlineAsyncTaskGooglePlayQueryInAppPurchases.cpp b/Engine/Plugins/Online/Android/OnlineSubsystemGooglePlay/Source/Private/OnlineAsyncTaskGooglePlayQueryInAppPurchases.cpp
deleted file mode 100644
diff --git a/Engine/Plugins/Online/Android/OnlineSubsystemGooglePlay/Source/Private/OnlineAsyncTaskGooglePlayQueryInAppPurchases.h b/Engine/Plugins/Online/Android/OnlineSubsystemGooglePlay/Source/Private/OnlineAsyncTaskGooglePlayQueryInAppPurchases.h
deleted file mode 100644
diff --git a/Engine/Plugins/Online/Android/OnlineSubsystemGooglePlay/Source/Private/OnlinePurchaseGooglePlay.cpp b/Engine/Plugins/Online/Android/OnlineSubsystemGooglePlay/Source/Private/OnlinePurchaseGooglePlay.cpp
--- a/Engine/Plugins/Online/Android/OnlineSubsystemGooglePlay/Source/Private/OnlinePurchaseGooglePlay.cpp
+++ b/Engine/Plugins/Online/Android/OnlineSubsystemGooglePlay/Source/Private/OnlinePurchaseGooglePlay.cpp
@@ -2,22 +2,19 @@
 
 #include "OnlinePurchaseGooglePlay.h"
 #include "OnlineSubsystemGooglePlay.h"
-#include "OnlineStoreGooglePlay.h"
 #include "OnlineError.h"
-#include "Misc/Base64.h"
-#include "Async/TaskGraphInterfaces.h"
 #include <jni.h>
 #include "Android/AndroidJavaEnv.h"
 
 #define LOCTEXT_NAMESPACE "OnlineSubsystemGooglePlay"
-#define GOOGLEPLAYUSER TEXT("GooglePlayUser")
 
-FGoogleTransactionData::FGoogleTransactionData(const FString& InOfferId, const FString& InProductToken, const FString& InReceiptData, const FString& InSignature)
+FGoogleTransactionData::FGoogleTransactionData(const FString& InOfferId, const FString& InPurchaseToken, const FString& InReceiptData, const FString& InSignature, EGooglePlayPurchaseState InPurchaseState)
 	: OfferId(InOfferId)
-	, TransactionIdentifier(InProductToken)
+	, PurchaseToken(InPurchaseToken)
 	, CombinedTransactionData(InReceiptData, InSignature)
+	, PurchaseState(InPurchaseState)
 {
-	if (TransactionIdentifier.IsEmpty())
+	if (PurchaseToken.IsEmpty())
 	{
 		ErrorStr = TEXT("Receipt does not contain purchase token");
 	}
@@ -31,6 +28,16 @@ FGoogleTransactionData::FGoogleTransactionData(const FString& InOfferId, const F
 	}
 }
 
+FString FGoogleTransactionData::ToDebugString() const
+{
+	return FString::Printf(TEXT("OfferId: %s PurchaseState: %s TransactionId: %s ReceiptData: %s%s"),
+		*OfferId,
+		LexToString(PurchaseState),
+		*PurchaseToken,
+		*CombinedTransactionData.ToJson(),
+		ErrorStr.IsEmpty() ? TEXT("") : *FString::Printf(TEXT(" Error: %s"), *ErrorStr));
+}
+
 FOnlinePurchaseGooglePlay::FOnlinePurchaseGooglePlay(FOnlineSubsystemGooglePlay* InSubsystem)
 	: bQueryingReceipts(false)
 	, Subsystem(InSubsystem)
@@ -38,29 +45,140 @@ FOnlinePurchaseGooglePlay::FOnlinePurchaseGooglePlay(FOnlineSubsystemGooglePlay*
 	UE_LOG_ONLINE_PURCHASE(Verbose, TEXT( "FOnlinePurchaseGooglePlay::FOnlinePurchaseGooglePlay" ));
 }
 
-FOnlinePurchaseGooglePlay::FOnlinePurchaseGooglePlay()
+/**
+ * Info used to cache and track the order being currently purchased.
+ */
+class FOnlinePurchaseGooglePlay::FOnlinePurchaseInProgressTransaction
 {
-	UE_LOG_ONLINE_PURCHASE(Verbose, TEXT( "FOnlinePurchaseGooglePlay::FOnlinePurchaseGooglePlay" ));
+public:
+	FOnlinePurchaseInProgressTransaction(
+		const FPurchaseCheckoutRequest& InCheckoutRequest,
+		const FUniqueNetId& InUserId,
+		const EPurchaseTransactionState InPendingTransactionState,
+		const FOnPurchaseCheckoutComplete& InCheckoutCompleteDelegate
+		);
+	
+	/**
+	 * Generate a final receipt for all purchases made in this single transaction
+	 */
+	TSharedRef<FPurchaseReceipt> GenerateReceipt();
+	
+	/** Generate one off receipts for transactions initiated outside the current run of the application */
+	static TSharedRef<FPurchaseReceipt> GenerateReceipt(const FGoogleTransactionData& Transaction);
+	
+	/** Add the single completed transaction to this transaction */
+	bool AddCompletedOffer(EPurchaseTransactionState Result, const FGoogleTransactionData& Transaction);
+
+public:
+	
+	/** Checkout info for the in progress order */
+	FPurchaseCheckoutRequest CheckoutRequest;
+	/** User for the order in progress */
+	FUniqueNetIdGooglePlayRef UserId;
+	/** Delegate to call on completion */
+	FOnPurchaseCheckoutComplete CheckoutCompleteDelegate;
+	/** Tracks the current state of the order */
+	FPurchaseReceipt PurchaseInfo;
+};
+
+
+FOnlinePurchaseGooglePlay::FOnlinePurchaseInProgressTransaction::FOnlinePurchaseInProgressTransaction(
+	const FPurchaseCheckoutRequest& InCheckoutRequest,
+	const FUniqueNetId& InUserId,
+	const EPurchaseTransactionState InPendingTransactionState,
+	const FOnPurchaseCheckoutComplete& InCheckoutCompleteDelegate
+	)
+	: CheckoutRequest(InCheckoutRequest)
+	, UserId(FUniqueNetIdGooglePlay::Cast(InUserId.AsShared()))
+	, CheckoutCompleteDelegate(InCheckoutCompleteDelegate)
+{
+	PurchaseInfo.TransactionState = InPendingTransactionState;
 }
 
-FOnlinePurchaseGooglePlay::~FOnlinePurchaseGooglePlay()
+TSharedRef<FPurchaseReceipt> FOnlinePurchaseGooglePlay::FOnlinePurchaseInProgressTransaction::GenerateReceipt()
 {
-	if (Subsystem)
+	TSharedRef<FPurchaseReceipt> Receipt = MakeShared<FPurchaseReceipt>();
+	
+	Receipt->TransactionState = PurchaseInfo.TransactionState;
+	Receipt->TransactionId = PurchaseInfo.TransactionId;
+	
+	if(PurchaseInfo.TransactionState == EPurchaseTransactionState::Purchased ||
+	   PurchaseInfo.TransactionState == EPurchaseTransactionState::Deferred)
 	{
-		Subsystem->ClearOnGooglePlayProcessPurchaseCompleteDelegate_Handle(ProcessPurchaseResultDelegateHandle);
-		Subsystem->ClearOnGooglePlayQueryExistingPurchasesCompleteDelegate_Handle(QueryExistingPurchasesCompleteDelegateHandle);
+		Receipt->ReceiptOffers = PurchaseInfo.ReceiptOffers;
 	}
+	else
+	{
+		// Add the requested offers to the receipt in the event of an incomplete purchase.
+		for(const auto& RequestedOffer : CheckoutRequest.PurchaseOffers)
+		{
+			Receipt->AddReceiptOffer(RequestedOffer.OfferNamespace, RequestedOffer.OfferId, RequestedOffer.Quantity);
+		}
+	}
+	
+	return Receipt;
 }
 
-void FOnlinePurchaseGooglePlay::Init()
+TSharedRef<FPurchaseReceipt> FOnlinePurchaseGooglePlay::FOnlinePurchaseInProgressTransaction::GenerateReceipt(const FGoogleTransactionData& Transaction)
 {
-	UE_LOG_ONLINE_PURCHASE(Verbose, TEXT("FOnlinePurchaseGooglePlay::Init"));
+	TSharedRef<FPurchaseReceipt> Receipt = MakeShared<FPurchaseReceipt>();
+	
+	Receipt->TransactionState = Transaction.GetErrorStr().IsEmpty() ? EPurchaseTransactionState::Purchased : EPurchaseTransactionState::Failed;
+	Receipt->TransactionId = Transaction.GetTransactionIdentifier();
+	
+	if (Receipt->TransactionState == EPurchaseTransactionState::Purchased)
+	{
+		FPurchaseReceipt::FReceiptOfferEntry ReceiptEntry(TEXT(""), Transaction.GetOfferId(), 1);
 
-	FOnGooglePlayProcessPurchaseCompleteDelegate PurchaseCompleteDelegate = FOnGooglePlayProcessPurchaseCompleteDelegate::CreateThreadSafeSP(this, &FOnlinePurchaseGooglePlay::OnTransactionCompleteResponse);
-	ProcessPurchaseResultDelegateHandle = Subsystem->AddOnGooglePlayProcessPurchaseCompleteDelegate_Handle(PurchaseCompleteDelegate);
+		FPurchaseReceipt::FLineItemInfo& LineItem = ReceiptEntry.LineItems.Emplace_GetRef();
 
-	FOnGooglePlayQueryExistingPurchasesCompleteDelegate QueryExistingPurchasesCompleteDelegate = FOnGooglePlayQueryExistingPurchasesCompleteDelegate::CreateThreadSafeSP(this, &FOnlinePurchaseGooglePlay::OnQueryExistingPurchasesComplete);
-	QueryExistingPurchasesCompleteDelegateHandle = Subsystem->AddOnGooglePlayQueryExistingPurchasesCompleteDelegate_Handle(QueryExistingPurchasesCompleteDelegate);
+		LineItem.ItemName = Transaction.GetOfferId();
+		LineItem.UniqueId = Transaction.GetTransactionIdentifier();
+		if(Transaction.GetPurchaseState() == EGooglePlayPurchaseState::Purchased)
+		{
+			// Pending purchases should not be redeemable
+			LineItem.ValidationInfo = Transaction.GetCombinedReceiptData();
+		}
+		Receipt->AddReceiptOffer(MoveTemp(ReceiptEntry));
+	}
+	
+	return Receipt;
+}
+
+bool FOnlinePurchaseGooglePlay::FOnlinePurchaseInProgressTransaction::AddCompletedOffer(EPurchaseTransactionState Result, const FGoogleTransactionData& Transaction)
+{
+	bool bIsFromThisRequest = CheckoutRequest.PurchaseOffers.ContainsByPredicate([&Transaction](const FPurchaseCheckoutRequest::FPurchaseOfferEntry& Entry)
+	{
+		return (Transaction.GetOfferId() == Entry.OfferId);
+	});
+	
+	if(!bIsFromThisRequest)
+	{
+		return false;
+	}
+
+	FPurchaseReceipt::FReceiptOfferEntry Receipt(TEXT(""), Transaction.GetOfferId(), 1);
+
+	FPurchaseReceipt::FLineItemInfo& LineItem = Receipt.LineItems.Emplace_GetRef();
+	
+	LineItem.ItemName = Transaction.GetOfferId();
+	LineItem.UniqueId = Transaction.GetTransactionIdentifier();
+	if(Transaction.GetPurchaseState() == EGooglePlayPurchaseState::Purchased)
+	{
+		// Pending purchases should not be redeemable
+		LineItem.ValidationInfo = Transaction.GetCombinedReceiptData();
+	}
+	PurchaseInfo.AddReceiptOffer(MoveTemp(Receipt));
+	return true;
+}
+
+FOnlinePurchaseGooglePlay::~FOnlinePurchaseGooglePlay()
+{
+}
+
+void FOnlinePurchaseGooglePlay::Init()
+{
+	UE_LOG_ONLINE_PURCHASE(Verbose, TEXT("FOnlinePurchaseGooglePlay::Init"));
 }
 
 bool FOnlinePurchaseGooglePlay::IsAllowedToPurchase(const FUniqueNetId& UserId)
@@ -75,60 +193,60 @@ void FOnlinePurchaseGooglePlay::Checkout(const FUniqueNetId& UserId, const FPurc
 	bool bStarted = false;
 	FText ErrorMessage;
 
-	TSharedRef<FOnlinePurchasePendingTransactionGooglePlay> RequestedTransaction = MakeShareable(new FOnlinePurchasePendingTransactionGooglePlay(CheckoutRequest, UserId, EPurchaseTransactionState::NotStarted, Delegate));
+	if (InProgressTransaction != nullptr)
+	{
+		UE_LOG_ONLINE_PURCHASE(Warning, TEXT("Previous purchase finished without notification"));			
+		InProgressTransaction.Reset();
+	}
+
+	TSharedPtr<FOnlinePurchaseInProgressTransaction> RequestedTransaction = MakeShared<FOnlinePurchaseInProgressTransaction>(CheckoutRequest, UserId, EPurchaseTransactionState::NotStarted, Delegate);
 
 	if (IsAllowedToPurchase(UserId))
 	{
-		const FString UserIdStr = GOOGLEPLAYUSER;
-		const TSharedRef<FOnlinePurchasePendingTransactionGooglePlay>* UserPendingTransaction = PendingTransactions.Find(UserIdStr);
-		if (UserPendingTransaction == nullptr)
+		int32 NumOffers = CheckoutRequest.PurchaseOffers.Num();
+		if (NumOffers > 0)
 		{
-			FOnlineStoreGooglePlayV2Ptr StoreInterface = StaticCastSharedPtr<FOnlineStoreGooglePlayV2>(Subsystem->GetStoreV2Interface());
-			if (StoreInterface.IsValid())
-			{
-				TSharedRef<FOnlinePurchasePendingTransactionGooglePlay> PendingTransaction = PendingTransactions.Add(UserIdStr, RequestedTransaction);
+			const FPurchaseCheckoutRequest::FPurchaseOfferEntry& Offer = CheckoutRequest.PurchaseOffers[0];
 
-				int32 NumOffers = CheckoutRequest.PurchaseOffers.Num();
-				if (NumOffers > 0)
-				{
-					const FPurchaseCheckoutRequest::FPurchaseOfferEntry& Offer = CheckoutRequest.PurchaseOffers[0];
-
-					extern bool AndroidThunkCpp_Iap_BeginPurchase(const FString&, const FString&);
-					FString AccountId;
-					if (CheckoutRequest.AccountId.IsSet())
-					{
-						AccountId = CheckoutRequest.AccountId.GetValue();
-					}
-					bStarted = AndroidThunkCpp_Iap_BeginPurchase(Offer.OfferId, AccountId);
-					UE_LOG_ONLINE_PURCHASE(Display, TEXT("Created Transaction? - %s"),
-						bStarted ? TEXT("Created a transaction.") : TEXT("Failed to create a transaction."));
-
-					RequestedTransaction->PendingPurchaseInfo.TransactionState = bStarted ? EPurchaseTransactionState::Processing : EPurchaseTransactionState::Failed;
-					if (NumOffers > 1)
-					{
-						UE_LOG_ONLINE_PURCHASE(Warning, TEXT("GooglePlay supports purchasing one offer at a time, %d were requested and ignored"), NumOffers - 1);
-					}
-				}
-				else
-				{
-					ErrorMessage = NSLOCTEXT("GooglePlayPurchase", "ErrorNoOffersSpecified", "Failed to checkout, no offers given.");
-					RequestedTransaction->PendingPurchaseInfo.TransactionState = EPurchaseTransactionState::Failed;
-				}
+			extern bool AndroidThunkCpp_Iap_BeginPurchase(const FString&, const FString&);
+			FString AccountId;
+			if (CheckoutRequest.AccountId.IsSet())
+			{
+				AccountId = CheckoutRequest.AccountId.GetValue();
+			}
+			bStarted = AndroidThunkCpp_Iap_BeginPurchase(Offer.OfferId, AccountId);
+			if(bStarted)
+			{
+				UE_LOG_ONLINE_PURCHASE(Display, TEXT("Created a transaction."));
+				RequestedTransaction->PurchaseInfo.TransactionState = EPurchaseTransactionState::Processing;
+			}
+			else
+			{
+				UE_LOG_ONLINE_PURCHASE(Display, TEXT("Failed to create a transaction."));
+				RequestedTransaction->PurchaseInfo.TransactionState = EPurchaseTransactionState::Failed;
+			}
+			if (NumOffers > 1)
+			{
+				UE_LOG_ONLINE_PURCHASE(Warning, TEXT("Only supported purchasing one offer at a time, %d were requested and ignored"), NumOffers - 1);
 			}
 		}
 		else
 		{
-			ErrorMessage = NSLOCTEXT("GooglePlayPurchase", "ErrorTransactionInProgress", "Failed to checkout, user has in progress transaction.");
-			RequestedTransaction->PendingPurchaseInfo.TransactionState = EPurchaseTransactionState::Failed;
+			ErrorMessage = NSLOCTEXT("GooglePlayPurchase", "ErrorNoOffersSpecified", "Failed to checkout, no offers given.");
+			RequestedTransaction->PurchaseInfo.TransactionState = EPurchaseTransactionState::Failed;
 		}
 	}
 	else
 	{
 		ErrorMessage = NSLOCTEXT("GooglePlayPurchase", "ErrorPurchaseNotAllowed", "Failed to checkout, user not allowed to purchase.");
-		RequestedTransaction->PendingPurchaseInfo.TransactionState = EPurchaseTransactionState::Failed;
+		RequestedTransaction->PurchaseInfo.TransactionState = EPurchaseTransactionState::Failed;
 	}
 
-	if (!bStarted)
+	if(bStarted)
+	{
+		InProgressTransaction = MoveTemp(RequestedTransaction);
+	}
+	else
 	{
 		TSharedRef<FPurchaseReceipt> FailReceipt = RequestedTransaction->GenerateReceipt();
 
@@ -158,46 +276,40 @@ void FOnlinePurchaseGooglePlay::RedeemCode(const FUniqueNetId& UserId, const FRe
 {
 	// Not supported
 	FOnlineError Result;
-	Delegate.ExecuteIfBound(Result, MakeShareable(new FPurchaseReceipt()));
+	Delegate.ExecuteIfBound(Result, MakeShared<FPurchaseReceipt>());
 }
 
-void FOnlinePurchaseGooglePlay::QueryReceipts(const FUniqueNetId& UserId, bool bRestoreReceipts, const FOnQueryReceiptsComplete& Delegate)
+void FOnlinePurchaseGooglePlay::QueryReceipts(const FUniqueNetId& UserId, bool /*bRestoreReceipts*/, const FOnQueryReceiptsComplete& Delegate)
 {
-	bool bSuccess = false;
-	bool bTriggerDelegate = true;
 	if (!bQueryingReceipts)
 	{
 		/**
-		 * bRestoreReceipts is irrelevant because GooglePlay requires the client to consume purchase in order to make a new purchase
-		 * There is no concept of restore here, any purchase query will reveal any non consumed purchases. They will remain in the list until the the game
-		 * consumes it via FinalizePurchase().
+		 * bRestoreReceipts is irrelevant because there is no concept of restore here, any purchase query will reveal any non consumed purchases. 
+		 * They will remain in the list until the the game consumes it via FinalizePurchase().
 		 */
 		bQueryingReceipts = true;
 		QueryReceiptsComplete = Delegate;
 
 		extern bool AndroidThunkCpp_Iap_QueryExistingPurchases();
-		if (AndroidThunkCpp_Iap_QueryExistingPurchases())
-		{
-			bTriggerDelegate = false;
-			bSuccess = true;
-		}
-		else
+		if (!AndroidThunkCpp_Iap_QueryExistingPurchases())
 		{
 			UE_LOG_ONLINE_PURCHASE(Verbose, TEXT("FOnlinePurchaseGooglePlay::QueryReceipts failed to start query"));
+			
+			Subsystem->ExecuteNextTick([this]() {
+				FOnlineError Result(false);
+				QueryReceiptsComplete.ExecuteIfBound(Result);
+				QueryReceiptsComplete.Unbind();
+				bQueryingReceipts = false;
+			});
 		}
 	}
 	else
 	{
 		UE_LOG_ONLINE_PURCHASE(Verbose, TEXT("FOnlinePurchaseGooglePlay::QueryReceipts already in progress."));
-	}
-
-	if (bTriggerDelegate)
-	{
-		Subsystem->ExecuteNextTick([this, Delegate, bSuccess]() {
-			FOnlineError Result(bSuccess);
+		
+		Subsystem->ExecuteNextTick([Delegate = Delegate]() mutable {
+			FOnlineError Result(false);
 			Delegate.ExecuteIfBound(Result);
-			bQueryingReceipts = false;
-			QueryReceiptsComplete.Unbind();
 		});
 	}
 }
@@ -206,54 +318,32 @@ void FOnlinePurchaseGooglePlay::GetReceipts(const FUniqueNetId& UserId, TArray<F
 {
 	OutReceipts.Empty();
 
-	// Add the cached list of user purchases
-	FString UserIdStr = GOOGLEPLAYUSER;
-	const TArray<TSharedRef<FPurchaseReceipt>>* UserCompletedTransactionsPtr = CompletedTransactions.Find(UserIdStr);
-	if (UserCompletedTransactionsPtr != nullptr)
+	for(const TSharedRef<FPurchaseReceipt>& Transaction: KnownTransactions)
 	{
-		const TArray<TSharedRef<FPurchaseReceipt>>& UserCompletedTransactions = *UserCompletedTransactionsPtr;
-		for (int32 Idx = 0; Idx < UserCompletedTransactions.Num(); Idx++)
-		{
-			UE_LOG_ONLINE_PURCHASE(Verbose, TEXT("FOnlinePurchaseGooglePlay::GetReceipts - Adding UserCompletedTransaction to OutReceipts"));
-			const TSharedRef<FPurchaseReceipt>& Transaction = UserCompletedTransactions[Idx];
-			OutReceipts.Add(*Transaction);
-		}
+		OutReceipts.Add(*Transaction);
 	}
 	
-	// Add purchases completed while "offline"
-	for (int32 Idx = 0; Idx < OfflineTransactions.Num(); Idx++)
-	{
-		UE_LOG_ONLINE_PURCHASE(Verbose, TEXT("FOnlinePurchaseGooglePlay::GetReceipts - Adding OfflineTransaction to OutReceipts"));
-		OutReceipts.Add(*OfflineTransactions[Idx]);
-	}
-
-	UE_LOG_ONLINE_PURCHASE(Log, TEXT("FOnlinePurchaseGooglePlay::GetReceipts - Final Number of outreciepts OutReceipts - %d"), OutReceipts.Num());
+	UE_LOG_ONLINE_PURCHASE(Log, TEXT("FOnlinePurchaseGooglePlay::GetReceipts - Final Number of OutReceipts - %d"), OutReceipts.Num());
 }
 
 void FOnlinePurchaseGooglePlay::OnTransactionCompleteResponse(EGooglePlayBillingResponseCode InResponseCode, const FGoogleTransactionData& InTransactionData)
 {
-	UE_LOG_ONLINE_PURCHASE(Log, TEXT("FOnlinePurchaseGooglePlay::OnTransactionCompleteResponse %s"), ToString(InResponseCode));
-	UE_LOG_ONLINE_PURCHASE(Log, TEXT("3... Transaction: %s"), *InTransactionData.ToDebugString());
-	EPurchaseTransactionState Result = ConvertGPResponseCodeToPurchaseTransactionState(InResponseCode);
+	UE_LOG_ONLINE_PURCHASE(Log, TEXT("FOnlinePurchaseGooglePlay::OnTransactionCompleteResponse %s. Transaction: %s"), LexToString(InResponseCode), *InTransactionData.ToDebugString());
+	EPurchaseTransactionState Result = ConvertGPResponseToPurchaseTransactionState(InResponseCode, InTransactionData.GetPurchaseState());
 
-	FString UserIdStr = GOOGLEPLAYUSER;
-	const TSharedRef<FOnlinePurchasePendingTransactionGooglePlay>* UserPendingTransactionPtr = PendingTransactions.Find(UserIdStr);
-	if(UserPendingTransactionPtr != nullptr)
+	bool bIsInProgressTransaction = false;
+	
+	if(InProgressTransaction != nullptr)
 	{
-		const TSharedRef<FOnlinePurchasePendingTransactionGooglePlay> UserPendingTransaction = *UserPendingTransactionPtr;
-		const FString& ErrorStr = InTransactionData.GetErrorStr();
-
-		if (Result == EPurchaseTransactionState::Purchased || Result == EPurchaseTransactionState::Restored)
-		{
-			if (!UserPendingTransaction->AddCompletedOffer(Result, InTransactionData))
-			{
-				UE_LOG_ONLINE_PURCHASE(Verbose, TEXT("Offer not found in pending transaction"));
-				Result = EPurchaseTransactionState::Failed;
-			}
-		}
+		// Only supported purchasing one offer at a time
+		const FPurchaseCheckoutRequest::FPurchaseOfferEntry& Offer = InProgressTransaction->CheckoutRequest.PurchaseOffers[0];
+		bool bWasInProgressTransactionCanceled = (Result == EPurchaseTransactionState::Canceled && InTransactionData.GetOfferId().IsEmpty());
+		bIsInProgressTransaction = bWasInProgressTransactionCanceled || (Offer.OfferId == InTransactionData.GetOfferId());
+	}
 
-		UserPendingTransaction->PendingPurchaseInfo.TransactionState = Result;
-		UserPendingTransaction->PendingPurchaseInfo.TransactionId = InTransactionData.GetTransactionIdentifier();
+	if(bIsInProgressTransaction)
+	{
+		const FString& ErrorStr = InTransactionData.GetErrorStr();
 
 		FOnlineError FinalResult;
 		switch (Result)
@@ -276,83 +366,93 @@ void FOnlinePurchaseGooglePlay::OnTransactionCompleteResponse(EGooglePlayBilling
 			case EPurchaseTransactionState::Invalid:
 				FinalResult.SetFromErrorCode(TEXT("com.epicgames.purchase.invalid"));
 				FinalResult.ErrorMessage = !ErrorStr.IsEmpty() ? FText::FromString(ErrorStr) : LOCTEXT("GooglePlayInvalidState", "Invalid purchase result");
-				UserPendingTransaction->PendingPurchaseInfo.TransactionState = EPurchaseTransactionState::Invalid;
 				break;
 			default:
-				UE_LOG_ONLINE_PURCHASE(Warning, TEXT("Unexpected state after purchase %d"), (int)Result);
+				UE_LOG_ONLINE_PURCHASE(Warning, TEXT("Unexpected EPurchaseTransactionState %d"), (int)Result);
 				FinalResult.SetFromErrorCode(TEXT("com.epicgames.purchase.unexpected_state"));
 				FinalResult.ErrorMessage = !ErrorStr.IsEmpty() ? FText::FromString(ErrorStr) : LOCTEXT("GooglePlayUnexpectedState", "Unexpected purchase result");
-				UserPendingTransaction->PendingPurchaseInfo.TransactionState = EPurchaseTransactionState::Failed;
 				break;
 		}
 
-		TSharedRef<FPurchaseReceipt> FinalReceipt = UserPendingTransaction->GenerateReceipt();
+		InProgressTransaction->PurchaseInfo.TransactionState = Result;
+		InProgressTransaction->PurchaseInfo.TransactionId = InTransactionData.GetTransactionIdentifier();
 
-		TArray< TSharedRef<FPurchaseReceipt> >& UserCompletedTransactions = CompletedTransactions.FindOrAdd(UserIdStr);
+		if (Result == EPurchaseTransactionState::Purchased ||
+			Result == EPurchaseTransactionState::Deferred )
+		{
+			InProgressTransaction->AddCompletedOffer(Result, InTransactionData);
+		}
 
-		PendingTransactions.Remove(UserIdStr);
+		TSharedRef<FPurchaseReceipt> FinalReceipt = InProgressTransaction->GenerateReceipt();
 
-		// If this is a deferred transaction, we will process it as an "offline" transaction, so don't complete
-		if (Result != EPurchaseTransactionState::Deferred)
+		if (Result == EPurchaseTransactionState::Purchased ||
+			Result == EPurchaseTransactionState::Deferred )
 		{
-			UserCompletedTransactions.Add(FinalReceipt);
+			KnownTransactions.Add(FinalReceipt);
 		}
 
-		UserPendingTransaction->CheckoutCompleteDelegate.ExecuteIfBound(FinalResult, FinalReceipt);
+		InProgressTransaction->CheckoutCompleteDelegate.ExecuteIfBound(FinalResult, FinalReceipt);
+		InProgressTransaction.Reset();
 	}
 	else
 	{
-		// Need to populate offlineTransactions here. 
-		// Transactions that come in during login or other non explicit purchase moments are added to a receipts list for later redemption
-		UE_LOG_ONLINE_PURCHASE(Log, TEXT("Pending transaction completed offline"));
-		if (Result == EPurchaseTransactionState::Restored || Result == EPurchaseTransactionState::Purchased)
+		if (Result == EPurchaseTransactionState::Purchased)
 		{
-			TSharedRef<FPurchaseReceipt> OfflineReceipt = FOnlinePurchasePendingTransactionGooglePlay::GenerateReceipt(InTransactionData);
-			OfflineTransactions.Add(OfflineReceipt);
-
-			// Queue this user to be updated about this next-tick on the game-thread, if they're not mid-purchase
-			TWeakPtr<FOnlinePurchaseGooglePlay, ESPMode::ThreadSafe> WeakThis(AsShared());
-			Subsystem->ExecuteNextTick([WeakThis]()
+			TSharedRef<FPurchaseReceipt> NewReceipt = FOnlinePurchaseInProgressTransaction::GenerateReceipt(InTransactionData);;
+			TSharedRef<FPurchaseReceipt>* ExistingReceipt = KnownTransactions.FindByPredicate([&InTransactionData](const TSharedRef<FPurchaseReceipt>& Receipt)
 			{
-				FOnlinePurchaseGooglePlayPtr StrongThis = WeakThis.Pin();
-				if (StrongThis.IsValid())
+				return Receipt->TransactionId == InTransactionData.GetTransactionIdentifier();
+			});
+			
+			if(ExistingReceipt)
+			{
+				if((*ExistingReceipt)->TransactionState == EPurchaseTransactionState::Deferred)
 				{
-					StrongThis->TriggerOnUnexpectedPurchaseReceiptDelegates(*FUniqueNetIdGooglePlay::EmptyId());
+					UE_LOG_ONLINE_PURCHASE(Log, TEXT("Pending transaction completed"));
 				}
-			});
+				else
+				{
+					UE_LOG_ONLINE_PURCHASE(Log, TEXT("Unexpected transaction completed (a possible pending transaction finished from previous session?)"));
+				}
+				*ExistingReceipt = NewReceipt;
+			}
+			else
+			{
+				KnownTransactions.Add(NewReceipt);
+			}
+			
+			TriggerOnUnexpectedPurchaseReceiptDelegates(*FUniqueNetIdGooglePlay::EmptyId());
 		}
 	}
 }
 
 void FOnlinePurchaseGooglePlay::OnQueryExistingPurchasesComplete(EGooglePlayBillingResponseCode InResponseCode, const TArray<FGoogleTransactionData>& InExistingPurchases)
 {
-	UE_LOG_ONLINE_PURCHASE(Log, TEXT("FOnlinePurchaseGooglePlay::OnQueryExistingPurchasesComplete Response: %s Num: %d"), ToString(InResponseCode), InExistingPurchases.Num());
+	UE_LOG_ONLINE_PURCHASE(Log, TEXT("FOnlinePurchaseGooglePlay::OnQueryExistingPurchasesComplete Response: %s Num: %d"), LexToString(InResponseCode), InExistingPurchases.Num());
 
 	if (bQueryingReceipts)
 	{
 		bool bSuccess = (InResponseCode == EGooglePlayBillingResponseCode::Ok);
 		if (bSuccess)
 		{
-			EPurchaseTransactionState Result = ConvertGPResponseCodeToPurchaseTransactionState(InResponseCode);
-			for (const FGoogleTransactionData& Purchase : InExistingPurchases)
+			KnownTransactions.Empty();
+
+			for(const FGoogleTransactionData& Purchase: InExistingPurchases)
 			{
 				UE_LOG_ONLINE_PURCHASE(Log, TEXT("Adding existing receipt %s"), *Purchase.ToDebugString());
-				TSharedRef<FPurchaseReceipt> OfflineReceipt = FOnlinePurchasePendingTransactionGooglePlay::GenerateReceipt(Purchase);
-				OfflineTransactions.Add(OfflineReceipt);
-			}
+				TSharedRef<FPurchaseReceipt> Receipt = FOnlinePurchaseInProgressTransaction::GenerateReceipt(Purchase); 
+				KnownTransactions.Add(Receipt);
+			}			
 		}
 		else
 		{
 			UE_LOG_ONLINE_PURCHASE(Log, TEXT("OnQueryExistingPurchasesComplete failed"));
 		}
 
-		Subsystem->ExecuteNextTick([this, bSuccess]()
-		{
-			FOnlineError Result(bSuccess);
-			QueryReceiptsComplete.ExecuteIfBound(Result);
-			bQueryingReceipts = false;
-			QueryReceiptsComplete.Unbind();
-		});
+		FOnlineError Result(bSuccess);
+		QueryReceiptsComplete.ExecuteIfBound(Result);
+		QueryReceiptsComplete.Unbind();
+		bQueryingReceipts = false;
 	}
 	else
 	{
@@ -360,229 +460,121 @@ void FOnlinePurchaseGooglePlay::OnQueryExistingPurchasesComplete(EGooglePlayBill
 	}
 }
 
-void FOnlinePurchaseGooglePlay::FinalizeReceiptValidationInfo(const FUniqueNetId& UserId, FString& InReceiptValidationInfo, const FOnFinalizeReceiptValidationInfoComplete& Delegate)
+void FOnlinePurchaseGooglePlay::OnConsumePurchaseComplete(EGooglePlayBillingResponseCode InResponseCode, const FString& TransactionIdentifier)
 {
-	FOnlineError DefaultSuccess(true);
-	Delegate.ExecuteIfBound(DefaultSuccess, InReceiptValidationInfo);
-}
+	UE_LOG_ONLINE_PURCHASE(Log, TEXT("FOnlinePurchaseGooglePlay::OnConsumePurchaseComplete Response: %s PurchaseToken: %s"), LexToString(InResponseCode), *TransactionIdentifier);
 
-TSharedRef<FPurchaseReceipt> FOnlinePurchasePendingTransactionGooglePlay::GenerateReceipt()
-{
-	TSharedRef<FPurchaseReceipt> Receipt = MakeShareable(new FPurchaseReceipt());
-	
-	Receipt->TransactionState = PendingPurchaseInfo.TransactionState;
-	Receipt->TransactionId = PendingPurchaseInfo.TransactionId;
-	
-	if(PendingPurchaseInfo.TransactionState == EPurchaseTransactionState::Purchased ||
-	   PendingPurchaseInfo.TransactionState == EPurchaseTransactionState::Restored)
-	{
-		Receipt->ReceiptOffers = PendingPurchaseInfo.ReceiptOffers;
-	}
-	else
+	if (InResponseCode == EGooglePlayBillingResponseCode::Ok || InResponseCode == EGooglePlayBillingResponseCode::ItemNotOwned)
 	{
-		// Add the requested offers to the receipt in the event of an incomplete purchase.
-		for(const auto& RequestedOffer : CheckoutRequest.PurchaseOffers)
+		KnownTransactions.RemoveAll([&TransactionIdentifier](const TSharedRef<FPurchaseReceipt>& Receipt) 
 		{
-			Receipt->AddReceiptOffer(RequestedOffer.OfferNamespace, RequestedOffer.OfferId, RequestedOffer.Quantity);
-		}
+			return Receipt->TransactionId == TransactionIdentifier;
+		});		
 	}
-	
-	return Receipt;
-}
-
-TSharedRef<FPurchaseReceipt> FOnlinePurchasePendingTransactionGooglePlay::GenerateReceipt(const FGoogleTransactionData& Transaction)
-{
-	TSharedRef<FPurchaseReceipt> Receipt = MakeShareable(new FPurchaseReceipt());
-	
-	Receipt->TransactionState = Transaction.GetErrorStr().IsEmpty() ? EPurchaseTransactionState::Purchased : EPurchaseTransactionState::Failed;
-	Receipt->TransactionId = Transaction.GetTransactionIdentifier();
-	
-	if (Receipt->TransactionState == EPurchaseTransactionState::Purchased ||
-		Receipt->TransactionState == EPurchaseTransactionState::Restored)
+	else
 	{
-		FPurchaseReceipt::FReceiptOfferEntry ReceiptEntry(TEXT(""), Transaction.GetOfferId(), 1);
-
-		int32 Idx = ReceiptEntry.LineItems.AddZeroed();
-
-		FPurchaseReceipt::FLineItemInfo& LineItem = ReceiptEntry.LineItems[Idx];
-
-		LineItem.ItemName = Transaction.GetOfferId();
-		LineItem.UniqueId = Transaction.GetTransactionIdentifier();
-		LineItem.ValidationInfo = Transaction.GetCombinedReceiptData();
-
-		Receipt->AddReceiptOffer(ReceiptEntry);
+		UE_LOG_ONLINE_PURCHASE(Log, TEXT("OnConsumePurchaseComplete failed"));
 	}
-	
-	return Receipt;
 }
 
-bool FOnlinePurchasePendingTransactionGooglePlay::AddCompletedOffer(EPurchaseTransactionState Result, const FGoogleTransactionData& Transaction)
+void FOnlinePurchaseGooglePlay::FinalizeReceiptValidationInfo(const FUniqueNetId& UserId, FString& InReceiptValidationInfo, const FOnFinalizeReceiptValidationInfoComplete& Delegate)
 {
-	for (int32 OfferIdx = 0; OfferIdx < CheckoutRequest.PurchaseOffers.Num(); ++OfferIdx)
-	{
-		const FPurchaseCheckoutRequest::FPurchaseOfferEntry& Offer = CheckoutRequest.PurchaseOffers[OfferIdx];
-		if (Transaction.GetOfferId() == Offer.OfferId)
-		{
-			FPurchaseReceipt::FReceiptOfferEntry Receipt(TEXT(""), Transaction.GetOfferId(), 1);
-
-			int32 Idx = Receipt.LineItems.AddZeroed();
-			
-			FPurchaseReceipt::FLineItemInfo& LineItem = Receipt.LineItems[Idx];
-			
-			LineItem.ItemName = Transaction.GetOfferId();
-			LineItem.UniqueId = Transaction.GetTransactionIdentifier();
-			LineItem.ValidationInfo = Transaction.GetCombinedReceiptData();
-			
-			PendingPurchaseInfo.AddReceiptOffer(Receipt);
-			return true;
-		}
-	}
-	
-	return false;
+	FOnlineError DefaultSuccess(true);
+	Delegate.ExecuteIfBound(DefaultSuccess, InReceiptValidationInfo);
 }
 
-JNI_METHOD void Java_com_epicgames_unreal_GooglePlayStoreHelper_nativeQueryExistingPurchasesComplete(JNIEnv* jenv, jobject thiz, jsize responseCode, jobjectArray ProductIDs, jobjectArray ProductTokens, jobjectArray ReceiptsData, jobjectArray Signatures)
+JNI_METHOD void Java_com_epicgames_unreal_GooglePlayStoreHelper_NativeQueryExistingPurchasesComplete(JNIEnv* jenv, jobject /*Thiz*/, jint ResponseCode, jobjectArray ProductIDs, jintArray PurchaseStatesArray, jobjectArray PurchaseTokens, jobjectArray ReceiptsData, jobjectArray Signatures)
 {
 	TArray<FGoogleTransactionData> ExistingPurchaseInfo;
 
-	EGooglePlayBillingResponseCode EGPResponse = (EGooglePlayBillingResponseCode)responseCode;
+	EGooglePlayBillingResponseCode EGPResponse = (EGooglePlayBillingResponseCode)ResponseCode;
 
 	bool bWasSuccessful = (EGPResponse == EGooglePlayBillingResponseCode::Ok);
-	if (jenv && bWasSuccessful)
+	UE_LOG_ONLINE_PURCHASE(Verbose, TEXT("NativeQueryExistingPurchasesComplete received response code: %s\n"), LexToString(EGPResponse));
+
+	if (bWasSuccessful)
 	{
 		jsize NumProducts = jenv->GetArrayLength(ProductIDs);
-		jsize NumProductTokens = jenv->GetArrayLength(ProductTokens);
+		jsize NumTransactionStates = jenv->GetArrayLength(PurchaseStatesArray);
+		jsize NumPurchaseTokens = jenv->GetArrayLength(PurchaseTokens);
 		jsize NumReceipts = jenv->GetArrayLength(ReceiptsData);
 		jsize NumSignatures = jenv->GetArrayLength(Signatures);
 
-		ensure((NumProducts == NumProductTokens) && (NumProducts == NumReceipts) && (NumProducts == NumSignatures));
+		ensure((NumProducts == NumPurchaseTokens) && (NumProducts == NumTransactionStates) && (NumProducts == NumReceipts) && (NumProducts == NumSignatures));
 
+		jint* PurchaseStates = jenv->GetIntArrayElements(PurchaseStatesArray, NULL);
+		
 		for (jsize Idx = 0; Idx < NumProducts; Idx++)
 		{
 			// Build the product information strings.
 			const auto OfferId = FJavaHelper::FStringFromLocalRef(jenv, (jstring)jenv->GetObjectArrayElement(ProductIDs, Idx));
-			const auto ProductToken = FJavaHelper::FStringFromLocalRef(jenv, (jstring)jenv->GetObjectArrayElement(ProductTokens, Idx));
+			const auto PurchaseToken = FJavaHelper::FStringFromLocalRef(jenv, (jstring)jenv->GetObjectArrayElement(PurchaseTokens, Idx));
 			const auto ReceiptData = FJavaHelper::FStringFromLocalRef(jenv, (jstring)jenv->GetObjectArrayElement(ReceiptsData, Idx));
 			const auto SignatureData = FJavaHelper::FStringFromLocalRef(jenv, (jstring)jenv->GetObjectArrayElement(Signatures, Idx));
+			const auto PurchaseState = (EGooglePlayPurchaseState)PurchaseStates[Idx];
 
-			FGoogleTransactionData ExistingPurchase(OfferId, ProductToken, ReceiptData, SignatureData);
+			FGoogleTransactionData ExistingPurchase(OfferId, PurchaseToken, ReceiptData, SignatureData, PurchaseState);
 			ExistingPurchaseInfo.Add(ExistingPurchase);
 
-			FPlatformMisc::LowLevelOutputDebugStringf(TEXT("\Existing Product Identifier: %s"), *ExistingPurchase.ToDebugString());
+			UE_LOG_ONLINE_PURCHASE(Verbose, TEXT("\Existing Product Identifier: %s"), *ExistingPurchase.ToDebugString());
 		}
+		jenv->ReleaseIntArrayElements(PurchaseStatesArray, PurchaseStates, JNI_ABORT);
 	}
-
-	DECLARE_CYCLE_STAT(TEXT("FSimpleDelegateGraphTask.QueryExistingPurchases"), STAT_FSimpleDelegateGraphTask_QueryExistingPurchases, STATGROUP_TaskGraphTasks);
-
-	FSimpleDelegateGraphTask::CreateAndDispatchWhenReady(
-		FSimpleDelegateGraphTask::FDelegate::CreateLambda([=]()
+	
+	if (auto OnlineSubGP = static_cast<FOnlineSubsystemGooglePlay* const>(IOnlineSubsystem::Get(GOOGLEPLAY_SUBSYSTEM)))
 	{
-		FPlatformMisc::LowLevelOutputDebugStringf(TEXT("Query existing purchases was completed %s\n"), bWasSuccessful ? TEXT("successfully") : TEXT("unsuccessfully"));
-		if (IOnlineSubsystem* const OnlineSub = IOnlineSubsystem::Get(GOOGLEPLAY_SUBSYSTEM))
+		OnlineSubGP->ExecuteNextTick([OnlineSubGP, EGPResponse, ExistingPurchaseInfo = MoveTemp(ExistingPurchaseInfo)]()
 		{
-			FOnlineSubsystemGooglePlay* const OnlineSubGP = static_cast<FOnlineSubsystemGooglePlay* const>(OnlineSub);
-			if (OnlineSubGP)
-			{
-				OnlineSubGP->TriggerOnGooglePlayQueryExistingPurchasesCompleteDelegates(EGPResponse, ExistingPurchaseInfo);
-			}
-		}
-	}),
-		GET_STATID(STAT_FSimpleDelegateGraphTask_QueryExistingPurchases),
-		nullptr,
-		ENamedThreads::GameThread
-		);
+			TSharedPtr<FOnlinePurchaseGooglePlay> PurchaseInt = StaticCastSharedPtr<FOnlinePurchaseGooglePlay>(OnlineSubGP->GetPurchaseInterface());
+			PurchaseInt->OnQueryExistingPurchasesComplete(EGPResponse, ExistingPurchaseInfo);
+		});
+	}
 }
 
-JNI_METHOD void Java_com_epicgames_unreal_GooglePlayStoreHelper_nativeRestorePurchasesComplete(JNIEnv* jenv, jobject thiz, jsize responseCode, jobjectArray ProductIDs, jobjectArray ProductTokens, jobjectArray ReceiptsData, jobjectArray Signatures)
+JNI_METHOD void Java_com_epicgames_unreal_GooglePlayStoreHelper_NativePurchaseComplete(JNIEnv* jenv, jobject /*Thiz*/, jint JavaResponseCode, jstring JavaProductId, jint JavaPurchaseState, jstring JavaPurchaseToken, jstring JavaReceiptData, jstring JavaSignature)
 {
-	TArray<FGoogleTransactionData> RestoredPurchaseInfo;
+	FString ProductId, PurchaseToken, ReceiptData, Signature;
+	EGooglePlayBillingResponseCode EGPResponse = (EGooglePlayBillingResponseCode)JavaResponseCode;
+	EGooglePlayPurchaseState PurchaseState = (EGooglePlayPurchaseState)JavaPurchaseState;
 
-	EGooglePlayBillingResponseCode EGPResponse = (EGooglePlayBillingResponseCode)responseCode;
-	bool bWasSuccessful = (EGPResponse == EGooglePlayBillingResponseCode::Ok);
+	ProductId = FJavaHelper::FStringFromParam(jenv, JavaProductId);
+	PurchaseToken = FJavaHelper::FStringFromParam(jenv, JavaPurchaseToken);
+	ReceiptData = FJavaHelper::FStringFromParam(jenv, JavaReceiptData);
+	Signature = FJavaHelper::FStringFromParam(jenv, JavaSignature);
 
-	if (jenv && bWasSuccessful)
-	{
-		jsize NumProducts = jenv->GetArrayLength(ProductIDs);
-		jsize NumProductTokens = jenv->GetArrayLength(ProductTokens);
-		jsize NumReceipts = jenv->GetArrayLength(ReceiptsData);
-		jsize NumSignatures = jenv->GetArrayLength(Signatures);
+	FGoogleTransactionData TransactionData(ProductId, PurchaseToken, ReceiptData, Signature, PurchaseState);
 
-		ensure((NumProducts == NumProductTokens) && (NumProducts == NumReceipts) && (NumProducts == NumSignatures));
+	UE_LOG_ONLINE_PURCHASE(Verbose, TEXT("NativePurchaseComplete received response code: %s. Transaction: %s\n"), LexToString(EGPResponse), *TransactionData.ToDebugString());
 
-		for (jsize Idx = 0; Idx < NumProducts; Idx++)
-		{
-			const auto OfferId = FJavaHelper::FStringFromLocalRef(jenv, (jstring)jenv->GetObjectArrayElement(ProductIDs, Idx));
-			const auto ProductToken = FJavaHelper::FStringFromLocalRef(jenv, (jstring)jenv->GetObjectArrayElement(ProductTokens, Idx));
-			const auto ReceiptData = FJavaHelper::FStringFromLocalRef(jenv, (jstring)jenv->GetObjectArrayElement(ReceiptsData, Idx));
-			const auto SignatureData = FJavaHelper::FStringFromLocalRef(jenv, (jstring)jenv->GetObjectArrayElement(Signatures, Idx));
-
-			FGoogleTransactionData RestoredPurchase(OfferId, ProductToken, ReceiptData, SignatureData);
-			RestoredPurchaseInfo.Add(RestoredPurchase);
-
-			FPlatformMisc::LowLevelOutputDebugStringf(TEXT("Restored Transaction: %s"), *RestoredPurchase.ToDebugString());
-		}
-	}
-
-	DECLARE_CYCLE_STAT(TEXT("FSimpleDelegateGraphTask.RestorePurchases"), STAT_FSimpleDelegateGraphTask_RestorePurchases, STATGROUP_TaskGraphTasks);
-
-	FSimpleDelegateGraphTask::CreateAndDispatchWhenReady(
-		FSimpleDelegateGraphTask::FDelegate::CreateLambda([=]()
+	if (auto OnlineSubGP = static_cast<FOnlineSubsystemGooglePlay* const>(IOnlineSubsystem::Get(GOOGLEPLAY_SUBSYSTEM)))
 	{
-		FPlatformMisc::LowLevelOutputDebugStringf(TEXT("Restoring In-App Purchases was completed  %s\n"), bWasSuccessful ? TEXT("successfully") : TEXT("unsuccessfully"));
-
-		if (IOnlineSubsystem* const OnlineSub = IOnlineSubsystem::Get(GOOGLEPLAY_SUBSYSTEM))
+		OnlineSubGP->ExecuteNextTick([OnlineSubGP, EGPResponse, TransactionData = MoveTemp(TransactionData)]()
 		{
-			FOnlineSubsystemGooglePlay* const OnlineSubGP = static_cast<FOnlineSubsystemGooglePlay* const>(OnlineSub);
-			if (OnlineSubGP)
-			{
-				OnlineSubGP->TriggerOnGooglePlayRestorePurchasesCompleteDelegates(EGPResponse, RestoredPurchaseInfo);
-			}
-		}
-	}),
-		GET_STATID(STAT_FSimpleDelegateGraphTask_RestorePurchases),
-		nullptr,
-		ENamedThreads::GameThread
-		);
+			TSharedPtr<FOnlinePurchaseGooglePlay> PurchaseInt = StaticCastSharedPtr<FOnlinePurchaseGooglePlay>(OnlineSubGP->GetPurchaseInterface());
+			PurchaseInt->OnTransactionCompleteResponse(EGPResponse, TransactionData);
+		});
+	}
 }
 
-JNI_METHOD void Java_com_epicgames_unreal_GooglePlayStoreHelper_nativePurchaseComplete(JNIEnv* jenv, jobject thiz, jsize responseCode, jstring productId, jstring productToken, jstring receiptData, jstring signature)
+JNI_METHOD void Java_com_epicgames_unreal_GooglePlayStoreHelper_NativeConsumeComplete(JNIEnv* jenv, jobject /*Thiz*/, jint JavaResponseCode, jstring JavaPurchaseToken)
 {
-	FString ProductId, ProductToken, ReceiptData, Signature;
-	EGooglePlayBillingResponseCode EGPResponse = (EGooglePlayBillingResponseCode)responseCode;
+	FString PurchaseToken;
+	EGooglePlayBillingResponseCode EGPResponse = (EGooglePlayBillingResponseCode)JavaResponseCode;
 
 	bool bWasSuccessful = (EGPResponse == EGooglePlayBillingResponseCode::Ok);
+	UE_LOG_ONLINE_PURCHASE(Verbose, TEXT("NativeConsumeComplete received response code: %s. PurchaseToken: %s\n"), LexToString(EGPResponse), *PurchaseToken);
 
-	//Store off results, because we will use them to determine if this is a deferred transaction.
-	ProductId = FJavaHelper::FStringFromParam(jenv, productId);
-	ProductToken = FJavaHelper::FStringFromParam(jenv, productToken);
-	ReceiptData = FJavaHelper::FStringFromParam(jenv, receiptData);
-	Signature = FJavaHelper::FStringFromParam(jenv, signature);
-
-
-	FGoogleTransactionData TransactionData(ProductId, ProductToken, ReceiptData, Signature);
+	PurchaseToken = FJavaHelper::FStringFromParam(jenv, JavaPurchaseToken);
 
-	FPlatformMisc::LowLevelOutputDebugStringf(TEXT("1... Response: %s, Transaction %s"), ToString(EGPResponse), *TransactionData.ToDebugString());
-
-	DECLARE_CYCLE_STAT(TEXT("FSimpleDelegateGraphTask.ProcessIapResult"), STAT_FSimpleDelegateGraphTask_ProcessIapResult, STATGROUP_TaskGraphTasks);
-
-	FSimpleDelegateGraphTask::CreateAndDispatchWhenReady(
-		FSimpleDelegateGraphTask::FDelegate::CreateLambda([=]()
+	if (auto OnlineSubGP = static_cast<FOnlineSubsystemGooglePlay* const>(IOnlineSubsystem::Get(GOOGLEPLAY_SUBSYSTEM)))
 	{
-		FPlatformMisc::LowLevelOutputDebugStringf(TEXT("In-App Purchase was completed  %s\n"), bWasSuccessful ? TEXT("successfully") : TEXT("unsuccessfully"));
-		if (IOnlineSubsystem* const OnlineSub = IOnlineSubsystem::Get(GOOGLEPLAY_SUBSYSTEM))
+		OnlineSubGP->ExecuteNextTick([OnlineSubGP, EGPResponse, PurchaseToken]()
 		{
-			FOnlineSubsystemGooglePlay* const OnlineSubGP = static_cast<FOnlineSubsystemGooglePlay* const>(OnlineSub);
-			if (OnlineSubGP)
-			{
-				FPlatformMisc::LowLevelOutputDebugStringf(TEXT("2... Response %s Transaction %s"), ToString(EGPResponse), *TransactionData.ToDebugString());
-				OnlineSubGP->TriggerOnGooglePlayProcessPurchaseCompleteDelegates(EGPResponse, TransactionData);
-			}
-		}
-	}),
-		GET_STATID(STAT_FSimpleDelegateGraphTask_ProcessIapResult),
-		nullptr,
-		ENamedThreads::GameThread
-		);
+			TSharedPtr<FOnlinePurchaseGooglePlay> PurchaseInt = StaticCastSharedPtr<FOnlinePurchaseGooglePlay>(OnlineSubGP->GetPurchaseInterface());
+			PurchaseInt->OnConsumePurchaseComplete(EGPResponse, PurchaseToken);
+		});
+	}
 }
 
 #undef LOCTEXT_NAMESPACE
+
diff --git a/Engine/Plugins/Online/Android/OnlineSubsystemGooglePlay/Source/Private/OnlineStoreGooglePlay.cpp b/Engine/Plugins/Online/Android/OnlineSubsystemGooglePlay/Source/Private/OnlineStoreGooglePlay.cpp
--- a/Engine/Plugins/Online/Android/OnlineSubsystemGooglePlay/Source/Private/OnlineStoreGooglePlay.cpp
+++ b/Engine/Plugins/Online/Android/OnlineSubsystemGooglePlay/Source/Private/OnlineStoreGooglePlay.cpp
@@ -2,7 +2,6 @@
 
 #include "OnlineStoreGooglePlay.h"
 #include "OnlineSubsystemGooglePlay.h"
-#include "OnlineAsyncTaskGooglePlayQueryInAppPurchases.h"
 
 #include "Internationalization/Internationalization.h"
 #include "Internationalization/Culture.h"
@@ -10,7 +9,6 @@
 #include "Misc/ConfigCacheIni.h"
 #include <jni.h>
 #include "Android/AndroidJavaEnv.h"
-#include "Async/TaskGraphInterfaces.h"
 
 FOnlineStoreGooglePlayV2::FOnlineStoreGooglePlayV2(FOnlineSubsystemGooglePlay* InSubsystem)
 	: bIsQueryInFlight(false)
@@ -28,19 +26,12 @@ FOnlineStoreGooglePlayV2::FOnlineStoreGooglePlayV2()
 
 FOnlineStoreGooglePlayV2::~FOnlineStoreGooglePlayV2()
 {
-	if (Subsystem)
-	{
-		Subsystem->ClearOnGooglePlayAvailableIAPQueryCompleteDelegate_Handle(AvailableIAPQueryDelegateHandle);
-	}
 }
 
 void FOnlineStoreGooglePlayV2::Init()
 {
 	UE_LOG_ONLINE_STOREV2(Verbose, TEXT("FOnlineStoreGooglePlayV2::Init"));
 
-	FOnGooglePlayAvailableIAPQueryCompleteDelegate Delegate = FOnGooglePlayAvailableIAPQueryCompleteDelegate::CreateThreadSafeSP(this, &FOnlineStoreGooglePlayV2::OnGooglePlayAvailableIAPQueryComplete);
-	AvailableIAPQueryDelegateHandle = Subsystem->AddOnGooglePlayAvailableIAPQueryCompleteDelegate_Handle(Delegate);
-
 	FString GooglePlayLicenseKey;
 	if (!GConfig->GetString(TEXT("/Script/AndroidRuntimeSettings.AndroidRuntimeSettings"), TEXT("GooglePlayLicenseKey"), GooglePlayLicenseKey, GEngineIni) || GooglePlayLicenseKey.IsEmpty())
 	{
@@ -58,35 +49,42 @@ TSharedRef<FOnlineStoreOffer> ConvertProductToStoreOffer(const FOnlineStoreOffer
 
 void FOnlineStoreGooglePlayV2::OnGooglePlayAvailableIAPQueryComplete(EGooglePlayBillingResponseCode InResponseCode, const TArray<FProvidedProductInformation>& InProvidedProductInformation)
 { 
+	UE_LOG_ONLINE_STOREV2(Verbose, TEXT("OnGooglePlayAvailableIAPQueryComplete Response: %s NumProucts: %d"), LexToString(InResponseCode), InProvidedProductInformation.Num());
+
+	bool bSuccess = (InResponseCode == EGooglePlayBillingResponseCode::Ok);
 	TArray<FUniqueOfferId> OfferIds;
-	for (const FProvidedProductInformation& Product : InProvidedProductInformation)
+	FString ErrorStr;
+	
+	if(!bIsQueryInFlight)
 	{
-		TSharedRef<FOnlineStoreOffer> NewProductOffer = ConvertProductToStoreOffer(Product);
-
-		AddOffer(NewProductOffer);
-		OfferIds.Add(NewProductOffer->OfferId);
-
-		UE_LOG_ONLINE_STOREV2(Log, TEXT("Product Identifier: %s, Name: %s, Desc: %s, Long Desc: %s, Price: %s IntPrice: %d"),
-			*NewProductOffer->OfferId,
-			*NewProductOffer->Title.ToString(),
-			*NewProductOffer->Description.ToString(),
-			*NewProductOffer->LongDescription.ToString(),
-			*NewProductOffer->PriceText.ToString(),
-			NewProductOffer->NumericPrice);
+		UE_LOG_ONLINE_STOREV2(Log, TEXT("OnGooglePlayAvailableIAPQueryComplete: No IAP query in flight"));
 	}
 
-	if (CurrentQueryTask)
+	if (bSuccess)
 	{
-		CurrentQueryTask->ProcessQueryAvailablePurchasesResults(InResponseCode);
-
-		// clear the pointer, it will be destroyed by the async task manager
-		CurrentQueryTask = nullptr;
+		for (const FProvidedProductInformation& Product : InProvidedProductInformation)
+		{
+			TSharedRef<FOnlineStoreOffer> NewProductOffer = ConvertProductToStoreOffer(Product);
+
+			AddOffer(NewProductOffer);
+			OfferIds.Add(NewProductOffer->OfferId);
+
+			UE_LOG_ONLINE_STOREV2(Log, TEXT("Product Identifier: %s, Name: %s, Desc: %s, Long Desc: %s, Price: %s IntPrice: %d"),
+				*NewProductOffer->OfferId,
+				*NewProductOffer->Title.ToString(),
+				*NewProductOffer->Description.ToString(),
+				*NewProductOffer->LongDescription.ToString(),
+				*NewProductOffer->PriceText.ToString(),
+				NewProductOffer->NumericPrice);
+		}
 	}
 	else
 	{
-		UE_LOG_ONLINE_STOREV2(Log, TEXT("OnGooglePlayAvailableIAPQueryComplete: No IAP query task in flight"));
+		ErrorStr = LexToString(InResponseCode);
 	}
-
+	
+	QueryOnlineStoreOffersCompleteDelegate.ExecuteIfBound(bSuccess, OfferIds, ErrorStr);
+	QueryOnlineStoreOffersCompleteDelegate.Unbind();
 	bIsQueryInFlight = false;
 }
 
@@ -119,13 +117,10 @@ void FOnlineStoreGooglePlayV2::QueryOffersById(const FUniqueNetId& UserId, const
 	}
 	else
 	{
-	CurrentQueryTask = new FOnlineAsyncTaskGooglePlayQueryInAppPurchasesV2(
-		Subsystem,
-		OfferIds,
-		Delegate);
-	Subsystem->QueueAsyncTask(CurrentQueryTask);
-
 		bIsQueryInFlight = true;
+		QueryOnlineStoreOffersCompleteDelegate = Delegate;
+		extern bool AndroidThunkCpp_Iap_QueryInAppPurchases(const TArray<FString>&);
+		AndroidThunkCpp_Iap_QueryInAppPurchases(OfferIds);
 	}
 }
 
@@ -165,46 +160,43 @@ TSharedPtr<FOnlineStoreOffer> FOnlineStoreGooglePlayV2::GetOffer(const FUniqueOf
 	return Result;
 }
 
-JNI_METHOD void Java_com_epicgames_unreal_GooglePlayStoreHelper_nativeQueryComplete(JNIEnv* jenv, jobject thiz, jsize responseCode, jobjectArray productIDs, jobjectArray titles, jobjectArray descriptions, jobjectArray prices, jfloatArray pricesRaw, jobjectArray currencyCodes, jobjectArray originalJson)
+JNI_METHOD void Java_com_epicgames_unreal_GooglePlayStoreHelper_NativeQueryComplete(JNIEnv* jenv, jobject /*Thiz*/, jint ResponseCode, jobjectArray ProductIDs, jobjectArray Titles, jobjectArray Descriptions, jobjectArray Prices, jfloatArray PriceValuesRaw, jobjectArray CurrencyCodes)
 {
 	TArray<FOnlineStoreOffer> ProvidedProductInformation;
-	EGooglePlayBillingResponseCode EGPResponse = (EGooglePlayBillingResponseCode)responseCode;
+	EGooglePlayBillingResponseCode EGPResponse = (EGooglePlayBillingResponseCode)ResponseCode;
 	bool bWasSuccessful = (EGPResponse == EGooglePlayBillingResponseCode::Ok);
 
 	if (jenv && bWasSuccessful)
 	{
-		jsize NumProducts = jenv->GetArrayLength(productIDs);
-		jsize NumTitles = jenv->GetArrayLength(titles);
-		jsize NumDescriptions = jenv->GetArrayLength(descriptions);
-		jsize NumPrices = jenv->GetArrayLength(prices);
-		jsize NumPricesRaw = jenv->GetArrayLength(pricesRaw);
-		jsize NumCurrencyCodes = jenv->GetArrayLength(currencyCodes);
-		jsize NumJsonStrings = jenv->GetArrayLength(originalJson);
+		jsize NumProducts = jenv->GetArrayLength(ProductIDs);
+		jsize NumTitles = jenv->GetArrayLength(Titles);
+		jsize NumDescriptions = jenv->GetArrayLength(Descriptions);
+		jsize NumPrices = jenv->GetArrayLength(Prices);
+		jsize NumPricesRaw = jenv->GetArrayLength(PriceValuesRaw);
+		jsize NumCurrencyCodes = jenv->GetArrayLength(CurrencyCodes);
 
-		ensure((NumProducts == NumTitles) && (NumProducts == NumDescriptions) && (NumProducts == NumPrices) && (NumProducts == NumPricesRaw) && (NumProducts == NumCurrencyCodes) && (NumProducts == NumJsonStrings));
+		ensure((NumProducts == NumTitles) && (NumProducts == NumDescriptions) && (NumProducts == NumPrices) && (NumProducts == NumPricesRaw) && (NumProducts == NumCurrencyCodes));
 
-		jfloat* PricesRaw = jenv->GetFloatArrayElements(pricesRaw, 0);
+		jfloat* PriceValues = jenv->GetFloatArrayElements(PriceValuesRaw, 0);
 
 		for (jsize Idx = 0; Idx < NumProducts; Idx++)
 		{
-			// Build the product information strings.
-
 			FOnlineStoreOffer NewProductInfo;
 
-			NewProductInfo.OfferId = FJavaHelper::FStringFromLocalRef(jenv, (jstring)jenv->GetObjectArrayElement(productIDs, Idx));
+			NewProductInfo.OfferId = FJavaHelper::FStringFromLocalRef(jenv, (jstring)jenv->GetObjectArrayElement(ProductIDs, Idx));
 
 			int32 OpenParenIdx = -1;
 			int32 CloseParenIdx = -1;
-			FString Title = FJavaHelper::FStringFromLocalRef(jenv, (jstring)jenv->GetObjectArrayElement(titles, Idx));
+			FString Title = FJavaHelper::FStringFromLocalRef(jenv, (jstring)jenv->GetObjectArrayElement(Titles, Idx));
 			if (Title.FindLastChar(TEXT(')'), CloseParenIdx) && Title.FindLastChar(TEXT('('), OpenParenIdx) && (OpenParenIdx < CloseParenIdx))
 			{
 				Title = Title.Left(OpenParenIdx).TrimEnd();
 			}
 			NewProductInfo.Title = FText::FromString(Title);
 
-			NewProductInfo.Description = FText::FromString(FJavaHelper::FStringFromLocalRef(jenv, (jstring)jenv->GetObjectArrayElement(descriptions, Idx)));
-			NewProductInfo.PriceText = FText::FromString(FJavaHelper::FStringFromLocalRef(jenv, (jstring)jenv->GetObjectArrayElement(prices, Idx)));
-			NewProductInfo.CurrencyCode = FJavaHelper::FStringFromLocalRef(jenv, (jstring)jenv->GetObjectArrayElement(currencyCodes, Idx));
+			NewProductInfo.Description = FText::FromString(FJavaHelper::FStringFromLocalRef(jenv, (jstring)jenv->GetObjectArrayElement(Descriptions, Idx)));
+			NewProductInfo.PriceText = FText::FromString(FJavaHelper::FStringFromLocalRef(jenv, (jstring)jenv->GetObjectArrayElement(Prices, Idx)));
+			NewProductInfo.CurrencyCode = FJavaHelper::FStringFromLocalRef(jenv, (jstring)jenv->GetObjectArrayElement(CurrencyCodes, Idx));
 
 			// Convert the backend stated price into its base units
 			FInternationalization& I18N = FInternationalization::Get();
@@ -212,21 +204,9 @@ JNI_METHOD void Java_com_epicgames_unreal_GooglePlayStoreHelper_nativeQueryCompl
 
 			const FDecimalNumberFormattingRules& FormattingRules = Culture.GetCurrencyFormattingRules(NewProductInfo.CurrencyCode);
 			const FNumberFormattingOptions& FormattingOptions = FormattingRules.CultureDefaultFormattingOptions;
-			double Val = static_cast<double>(PricesRaw[Idx])* static_cast<double>(FMath::Pow(10.0f, FormattingOptions.MaximumFractionalDigits));
+			double Val = static_cast<double>(PriceValues[Idx]) * static_cast<double>(FMath::Pow(10.0f, FormattingOptions.MaximumFractionalDigits));
 
 			NewProductInfo.NumericPrice = FMath::TruncToInt(Val + 0.5);
-
-			//Loop through original json data and populate dynamic map.
-			TSharedPtr<FJsonObject> JsonObject;
-			TSharedRef<TJsonReader<>> JsonReader = TJsonReaderFactory<>::Create(FJavaHelper::FStringFromLocalRef(jenv, (jstring)jenv->GetObjectArrayElement(originalJson, Idx)));
-			if (FJsonSerializer::Deserialize(JsonReader, JsonObject) && JsonObject.IsValid())
-			{
-				for (TPair<FString, TSharedPtr<FJsonValue>> JsonValue : JsonObject->Values)
-				{
-					NewProductInfo.DynamicFields.Add(JsonValue.Key, JsonValue.Value->AsString());
-				}
-			}
-
 			NewProductInfo.ReleaseDate = FDateTime::MinValue();
 			NewProductInfo.ExpirationDate = FDateTime::MaxValue();
 
@@ -240,28 +220,17 @@ JNI_METHOD void Java_com_epicgames_unreal_GooglePlayStoreHelper_nativeQueryCompl
 				NewProductInfo.NumericPrice,
 				*NewProductInfo.CurrencyCode);
 		}
+		jenv->ReleaseFloatArrayElements(PriceValuesRaw, PriceValues, JNI_ABORT);
 	}
 
-	DECLARE_CYCLE_STAT(TEXT("FSimpleDelegateGraphTask.ProcessQueryIapResult"), STAT_FSimpleDelegateGraphTask_ProcessQueryIapResult, STATGROUP_TaskGraphTasks);
-
-	FPlatformMisc::LowLevelOutputDebugStringf(TEXT("Adding task Success: %d Response: %s"), bWasSuccessful, ToString(EGPResponse));
+	UE_LOG_ONLINE_STOREV2(Verbose, TEXT("QueryOffersById result Success: %d Response: %s"), bWasSuccessful, LexToString(EGPResponse));
 
-	FSimpleDelegateGraphTask::CreateAndDispatchWhenReady(
-		FSimpleDelegateGraphTask::FDelegate::CreateLambda([=]()
+	if (auto OnlineSubGP = static_cast<FOnlineSubsystemGooglePlay* const>(IOnlineSubsystem::Get(GOOGLEPLAY_SUBSYSTEM)))
 	{
-		if (IOnlineSubsystem* const OnlineSub = IOnlineSubsystem::Get(GOOGLEPLAY_SUBSYSTEM))
+		OnlineSubGP->ExecuteNextTick([OnlineSubGP, EGPResponse, Response = MoveTemp(ProvidedProductInformation)]()
 		{
-			FOnlineSubsystemGooglePlay* const OnlineSubGP = static_cast<FOnlineSubsystemGooglePlay* const>(OnlineSub);
-			if (OnlineSubGP)
-			{
-				FPlatformMisc::LowLevelOutputDebugStringf(TEXT("TriggerOnGooglePlayAvailableIAPQueryCompleteDelegates %s Size: %d"), ToString(EGPResponse), ProvidedProductInformation.Num());
-				OnlineSubGP->TriggerOnGooglePlayAvailableIAPQueryCompleteDelegates(EGPResponse, ProvidedProductInformation);
-			}
-		}
-		FPlatformMisc::LowLevelOutputDebugStringf(TEXT("In-App Purchase query was completed  %s\n"), bWasSuccessful ? TEXT("successfully") : TEXT("unsuccessfully"));
-	}),
-		GET_STATID(STAT_FSimpleDelegateGraphTask_ProcessQueryIapResult),
-		nullptr,
-		ENamedThreads::GameThread
-		);
+			TSharedPtr<FOnlineStoreGooglePlayV2> StoreInt = StaticCastSharedPtr<FOnlineStoreGooglePlayV2>(OnlineSubGP->GetStoreV2Interface());
+			StoreInt->OnGooglePlayAvailableIAPQueryComplete(EGPResponse, Response);
+		});
+	}
 }
diff --git a/Engine/Plugins/Online/Android/OnlineSubsystemGooglePlay/Source/Private/OnlineStoreGooglePlayCommon.h b/Engine/Plugins/Online/Android/OnlineSubsystemGooglePlay/Source/Private/OnlineStoreGooglePlayCommon.h
--- a/Engine/Plugins/Online/Android/OnlineSubsystemGooglePlay/Source/Private/OnlineStoreGooglePlayCommon.h
+++ b/Engine/Plugins/Online/Android/OnlineSubsystemGooglePlay/Source/Private/OnlineStoreGooglePlayCommon.h
@@ -7,9 +7,16 @@
 
 struct FGoogleTransactionData;
 
-/** Possible responses returned from the Java GooglePlay billing interface */
-enum class EGooglePlayBillingResponseCode : uint8
+/** Possible responses returned from the Java GooglePlay billing interface.
+ *  CustomLogicError is a custom Unreal Engine error to represent logic errors 
+ *  on the Java side
+ */
+enum class EGooglePlayBillingResponseCode : int8
 {
+	CustomLogicError = -127, // Custom value. Should match with GooglePlayStoreHelper.CustomLogicErrorResponse on GooglePlayStoreHelper.java
+	ServiceTimeout = -3,
+	FeatureNotSupported = -2,
+	ServiceDisconnected = -1,
 	Ok = 0,
 	UserCancelled = 1,
 	ServiceUnavailable = 2,
@@ -21,10 +28,26 @@ enum class EGooglePlayBillingResponseCode : uint8
 	ItemNotOwned = 8,
 };
 
-inline const TCHAR* const ToString(EGooglePlayBillingResponseCode InResponseCode)
+/* Possible transaction states returned from the Java GooglePlay billing interface.*/
+enum class EGooglePlayPurchaseState : uint8
+{
+	UnspecifiedState = 0,
+	Purchased = 1,
+	Pending = 2,
+};
+
+inline const TCHAR* const LexToString(EGooglePlayBillingResponseCode InResponseCode)
 {
 	switch (InResponseCode)
 	{
+		case EGooglePlayBillingResponseCode::CustomLogicError:
+			return TEXT("CustomLogicError");
+		case EGooglePlayBillingResponseCode::ServiceTimeout:
+			return TEXT("ServiceTimeout");
+		case EGooglePlayBillingResponseCode::FeatureNotSupported:
+			return TEXT("FeatureNotSupported");
+		case EGooglePlayBillingResponseCode::ServiceDisconnected:
+			return TEXT("ServiceDisconnected");
 		case EGooglePlayBillingResponseCode::Ok:
 			return TEXT("Ok");
 		case EGooglePlayBillingResponseCode::UserCancelled:
@@ -48,6 +71,21 @@ inline const TCHAR* const ToString(EGooglePlayBillingResponseCode InResponseCode
 	}
 }
 
+inline const TCHAR* const LexToString(EGooglePlayPurchaseState InTransactionState)
+{
+	switch (InTransactionState)
+	{
+		case EGooglePlayPurchaseState::UnspecifiedState:
+			return TEXT("UnspecifiedState");
+		case EGooglePlayPurchaseState::Purchased:
+			return TEXT("Purchased");
+		case EGooglePlayPurchaseState::Pending:
+			return TEXT("Pending");
+		default:
+			return TEXT("UnknownState");
+	}
+}
+
 inline EInAppPurchaseState::Type ConvertGPResponseCodeToIAPState(const EGooglePlayBillingResponseCode InResponseCode)
 {
 	switch (InResponseCode)
@@ -60,6 +98,10 @@ inline EInAppPurchaseState::Type ConvertGPResponseCodeToIAPState(const EGooglePl
 			return EInAppPurchaseState::AlreadyOwned;
 		case EGooglePlayBillingResponseCode::ItemNotOwned:
 			return EInAppPurchaseState::NotAllowed;
+		case EGooglePlayBillingResponseCode::CustomLogicError:
+		case EGooglePlayBillingResponseCode::ServiceTimeout:
+		case EGooglePlayBillingResponseCode::FeatureNotSupported:
+		case EGooglePlayBillingResponseCode::ServiceDisconnected:
 		case EGooglePlayBillingResponseCode::ServiceUnavailable:
 		case EGooglePlayBillingResponseCode::BillingUnavailable:
 		case EGooglePlayBillingResponseCode::ItemUnavailable:
@@ -70,19 +112,36 @@ inline EInAppPurchaseState::Type ConvertGPResponseCodeToIAPState(const EGooglePl
 	}
 }
 
-inline EPurchaseTransactionState ConvertGPResponseCodeToPurchaseTransactionState(const EGooglePlayBillingResponseCode InResponseCode)
+inline EPurchaseTransactionState ConvertGPPurchaseStateToPurchaseTransactionState(const EGooglePlayPurchaseState InTransactionState)
+{
+	switch(InTransactionState)
+	{
+		case EGooglePlayPurchaseState::Purchased:
+			return EPurchaseTransactionState::Purchased;
+		case EGooglePlayPurchaseState::Pending:
+			return EPurchaseTransactionState::Deferred;
+		default:
+			return EPurchaseTransactionState::Invalid;
+	}
+}
+
+inline EPurchaseTransactionState ConvertGPResponseToPurchaseTransactionState(const EGooglePlayBillingResponseCode InResponseCode, const EGooglePlayPurchaseState InPurchaseState)
 {
 	switch (InResponseCode)
 	{
 		case EGooglePlayBillingResponseCode::Ok:
-			return EPurchaseTransactionState::Purchased;
+			return ConvertGPPurchaseStateToPurchaseTransactionState(InPurchaseState);
 		case EGooglePlayBillingResponseCode::UserCancelled:
 			return EPurchaseTransactionState::Canceled;
 		case EGooglePlayBillingResponseCode::ItemAlreadyOwned:
-			// Non consumable purchased again?
+			// Non consumable purchased again or consumable purchased again before consuming
 			return EPurchaseTransactionState::Invalid;
 		case EGooglePlayBillingResponseCode::ItemNotOwned:
 			return EPurchaseTransactionState::Invalid;
+		case EGooglePlayBillingResponseCode::CustomLogicError:
+		case EGooglePlayBillingResponseCode::ServiceTimeout:
+		case EGooglePlayBillingResponseCode::FeatureNotSupported:
+		case EGooglePlayBillingResponseCode::ServiceDisconnected:
 		case EGooglePlayBillingResponseCode::ServiceUnavailable:
 		case EGooglePlayBillingResponseCode::BillingUnavailable:
 		case EGooglePlayBillingResponseCode::ItemUnavailable:
@@ -93,17 +152,6 @@ inline EPurchaseTransactionState ConvertGPResponseCodeToPurchaseTransactionState
 	}
 }
 
-/**
- * The resulting state of an iap transaction
- */
-enum class EInAppPurchaseResult : uint8
-{
-	Succeeded = 0,
-	RestoredFromServer,
-	Failed,
-	Cancelled,
-};
-
 typedef FOnlineStoreOffer FProvidedProductInformation;
 
  /**
@@ -125,7 +173,7 @@ DECLARE_MULTICAST_DELEGATE_TwoParams(FOnGooglePlayProcessPurchaseComplete, EGoog
 typedef FOnGooglePlayProcessPurchaseComplete::FDelegate FOnGooglePlayProcessPurchaseCompleteDelegate;
 
 /**
- * Delegate fired internally an existing purchases query has completed
+ * Delegate fired internally when an existing purchases query has completed
  *
  * @param InResponseCode response from the GooglePlay backend
  * @param InExistingPurchases known purchases for the user (non consumed or permanent)
@@ -134,10 +182,10 @@ DECLARE_MULTICAST_DELEGATE_TwoParams(FOnGooglePlayQueryExistingPurchasesComplete
 typedef FOnGooglePlayQueryExistingPurchasesComplete::FDelegate FOnGooglePlayQueryExistingPurchasesCompleteDelegate;
 
 /**
- * Delegate fired internally when existing purchases have been restored (StoreV1 only)
+ * Delegate fired internally when a consume purchase query has completed
  *
  * @param InResponseCode response from the GooglePlay backend
- * @param InRestoredPurchases restored for the user (non consumed or permanent)
+ * @param InPurchaseToken purchase token for consumed purchase
  */
-DECLARE_MULTICAST_DELEGATE_TwoParams(FOnGooglePlayRestorePurchasesComplete, EGooglePlayBillingResponseCode /*InResponseCode*/, const TArray<FGoogleTransactionData>& /*InRestoredPurchases*/);
-typedef FOnGooglePlayRestorePurchasesComplete::FDelegate FOnGooglePlayRestorePurchasesCompleteDelegate;
+DECLARE_MULTICAST_DELEGATE_TwoParams(FOnGooglePlayConsumePurchaseComplete, EGooglePlayBillingResponseCode /*InResponseCode*/, const FString& /*InPurchaseToken*/);
+typedef FOnGooglePlayConsumePurchaseComplete::FDelegate FOnGooglePlayConsumePurchaseCompleteDelegate;
diff --git a/Engine/Plugins/Online/Android/OnlineSubsystemGooglePlay/Source/Public/OnlinePurchaseGooglePlay.h b/Engine/Plugins/Online/Android/OnlineSubsystemGooglePlay/Source/Public/OnlinePurchaseGooglePlay.h
--- a/Engine/Plugins/Online/Android/OnlineSubsystemGooglePlay/Source/Public/OnlinePurchaseGooglePlay.h
+++ b/Engine/Plugins/Online/Android/OnlineSubsystemGooglePlay/Source/Public/OnlinePurchaseGooglePlay.h
@@ -6,37 +6,32 @@
 #include "Serialization/JsonSerializerMacros.h"
 #include "OnlineIdentityInterfaceGooglePlay.h"
 
-enum class EGooglePlayBillingResponseCode : uint8;
+enum class EGooglePlayBillingResponseCode : int8;
+enum class EGooglePlayPurchaseState : uint8;
 
 /**
  * Holds in a common format the data that comes out of a Google purchase transaction
  */
 struct FGoogleTransactionData
 {
-	FGoogleTransactionData(const FString& InOfferId, const FString& InProductToken, const FString& InReceiptData, const FString& InSignature);
+	FGoogleTransactionData(const FString& InOfferId, const FString& InPurchaseToken, const FString& InReceiptData, const FString& InSignature, EGooglePlayPurchaseState InPurchaseState);
 
 	/** @return a string that prints useful debug information about this transaction */
-	FString ToDebugString() const
-	{
-		return FString::Printf(TEXT("OfferId: %s TransactionId: %s ReceiptData: %s%s"),
-			*OfferId,
-			*TransactionIdentifier,
-			*CombinedTransactionData.ToJson(),
-			ErrorStr.IsEmpty() ? TEXT("") : *FString::Printf(TEXT(" Error: %s"), *ErrorStr));
-	}
-
+	FString ToDebugString() const;
 	/** @return offer id for this transaction */
 	const FString& GetOfferId() const { return OfferId; }
 	/** @return receipt data for this transaction */
-	const FString GetCombinedReceiptData() const { return CombinedTransactionData.ToJson(); }
+	FString GetCombinedReceiptData() const { return CombinedTransactionData.ToJson(); }
 	/** @return receipt data for this transaction */
 	const FString& GetReceiptData() const { return CombinedTransactionData.ReceiptData; }
 	/** @return signature for this transaction */
 	const FString& GetSignature() const { return CombinedTransactionData.Signature; }
 	/** @return error string for this transaction, if applicable */
 	const FString& GetErrorStr() const { return ErrorStr; }
-	/** @return the transaction id */
-	const FString& GetTransactionIdentifier() const	{ return TransactionIdentifier; }
+	/** @return the purchase transaction id */
+	const FString& GetTransactionIdentifier() const	{ return PurchaseToken; }
+	/** @return the purchase state */
+	EGooglePlayPurchaseState GetPurchaseState() const { return PurchaseState; }
 
 private:
 
@@ -66,54 +61,14 @@ private:
 
 	/** GooglePlay offer id */
 	FString OfferId;
-	/** Unique transaction id (purchaseToken) */
-	FString TransactionIdentifier;
+	/** PurchaseToken for the transaction */
+	FString PurchaseToken;
 	/** Error on the transaction, if applicable */
 	FString ErrorStr;
 	/** Combined receipt with signature in JSON */
 	FJsonReceiptData CombinedTransactionData;
-};
-
-/**
- * Info used to cache and track orders in progress.
- */
-class FOnlinePurchasePendingTransactionGooglePlay
-{
-public:
-	FOnlinePurchasePendingTransactionGooglePlay(
-		const FPurchaseCheckoutRequest& InCheckoutRequest,
-		const FUniqueNetId& InUserId,
-		const EPurchaseTransactionState InPendingTransactionState = EPurchaseTransactionState::NotStarted,
-		const FOnPurchaseCheckoutComplete& InCheckoutCompleteDelegate = FOnPurchaseCheckoutComplete()
-		)
-		: CheckoutRequest(InCheckoutRequest)
-		, UserId(FUniqueNetIdGooglePlay::Cast(InUserId.AsShared()))
-		, CheckoutCompleteDelegate(InCheckoutCompleteDelegate)
-	{
-		PendingPurchaseInfo.TransactionState = InPendingTransactionState;
-	}
-	
-	/**
-	 * Generate a final receipt for all purchases made in this single transaction
-	 */
-	TSharedRef<FPurchaseReceipt> GenerateReceipt();
-	
-	/** Generate one off receipts for transactions initiated outside the current run of the application */
-	static TSharedRef<FPurchaseReceipt> GenerateReceipt(const FGoogleTransactionData& Transaction);
-	
-	/** Add the single completed transaction to this transaction */
-	bool AddCompletedOffer(EPurchaseTransactionState Result, const FGoogleTransactionData& Transaction);
-
-public:
-	
-	/** Checkout info for the pending order */
-	const FPurchaseCheckoutRequest CheckoutRequest;
-	/** User for the pending order */
-	const FUniqueNetIdGooglePlayRef UserId;
-	/** Delegate to call on completion */
-	const FOnPurchaseCheckoutComplete CheckoutCompleteDelegate;
-	/** Tracks the current state of the order */
-	FPurchaseReceipt PendingPurchaseInfo;
+	/** Reported GooglePlay transaction state */
+	EGooglePlayPurchaseState PurchaseState;
 };
 
 /**
@@ -146,11 +101,6 @@ public:
 	 */
 	FOnlinePurchaseGooglePlay(class FOnlineSubsystemGooglePlay* InSubsystem);
 
-	/**
-	 * Constructor
-	 */
-	FOnlinePurchaseGooglePlay();
-
 	/**
 	 * Destructor
 	 */
@@ -159,53 +109,40 @@ public:
 	/** Initialize the interface */
 	void Init();
 
+	/** Handle Java side transaction completed notification */
+	void OnTransactionCompleteResponse(EGooglePlayBillingResponseCode InResponseCode, const FGoogleTransactionData& InTransactionData);
+
+	/** Handle Java side query purchases completed notification */
+	void OnQueryExistingPurchasesComplete(EGooglePlayBillingResponseCode InResponseCode, const TArray<FGoogleTransactionData>& InExistingPurchases);
+
+	/** Handle Java side consume purchase completed notification */
+	void OnConsumePurchaseComplete(EGooglePlayBillingResponseCode InResponseCode, const FString& InPurchaseToken);
 private:
 	
-	/** Mapping from user id to pending transaction */
-	typedef TMap<FString, TSharedRef<FOnlinePurchasePendingTransactionGooglePlay> > FOnlinePurchasePendingTransactionMap;
+	/**
+	 * Info used to cache and track orders in progress.
+	 */
+	class FOnlinePurchaseInProgressTransaction;
 
-	/** Mapping from user id to complete transactions */
-	typedef TMap<FString, TArray< TSharedRef<FPurchaseReceipt> > > FOnlinePurchaseCompleteTransactionsMap;
-	
-	/** Array of transactions completion indirectly (previous run, etc) */
-	typedef TArray< TSharedRef<FPurchaseReceipt> > FOnlineCompletedTransactions;
+	/** Complete transactions */
+	using FOnlinePurchasePurchasedTransactions = TArray< TSharedRef<FPurchaseReceipt> >;
 	
 private:
 	
 	/** Are receipts being queried */
 	bool bQueryingReceipts;
+	
 	/** Transient delegate to fire when query receipts has completed */
 	FOnQueryReceiptsComplete QueryReceiptsComplete;
 
 	/** Keeps track of pending user transactions */
-	FOnlinePurchasePendingTransactionMap PendingTransactions;
-	
-	/** Cache of completed transactions */
-	FOnlinePurchaseCompleteTransactionsMap CompletedTransactions;
+	TSharedPtr<FOnlinePurchaseInProgressTransaction> InProgressTransaction;
 	
-	/** Cache of purchases completed outside the running instance */
-	FOnlineCompletedTransactions OfflineTransactions;
+	/** Cache of known transactions */
+	FOnlinePurchasePurchasedTransactions KnownTransactions;
 	
 	/** Reference to the parent subsystem */
 	FOnlineSubsystemGooglePlay* Subsystem;
-
-	/**
-	 * Delegate fired when a purchase has completed
-	 *
-	 * @param InResponseCode response from the GooglePlay backend
-	 * @param InTransactionData transaction data for the completed purchase
-	 */
-	void OnTransactionCompleteResponse(EGooglePlayBillingResponseCode InResponseCode, const FGoogleTransactionData& InTransactionData);
-	FDelegateHandle ProcessPurchaseResultDelegateHandle;
-
-	/**
-	 * Delegate fired when a query for existing purchase has completed
-	 *
-	 * @param InResponseCode response from the GooglePlay backend
-	 * @param InExistingPurchases known purchases for the user (non consumed or permanent)
-	 */
-	void OnQueryExistingPurchasesComplete(EGooglePlayBillingResponseCode InResponseCode, const TArray<FGoogleTransactionData>& InExistingPurchases);
-	FDelegateHandle QueryExistingPurchasesCompleteDelegateHandle;
 };
 
 typedef TSharedPtr<FOnlinePurchaseGooglePlay, ESPMode::ThreadSafe> FOnlinePurchaseGooglePlayPtr;
diff --git a/Engine/Plugins/Online/Android/OnlineSubsystemGooglePlay/Source/Public/OnlineStoreGooglePlay.h b/Engine/Plugins/Online/Android/OnlineSubsystemGooglePlay/Source/Public/OnlineStoreGooglePlay.h
--- a/Engine/Plugins/Online/Android/OnlineSubsystemGooglePlay/Source/Public/OnlineStoreGooglePlay.h
+++ b/Engine/Plugins/Online/Android/OnlineSubsystemGooglePlay/Source/Public/OnlineStoreGooglePlay.h
@@ -7,7 +7,6 @@
 
 class FOnlineSubsystemGooglePlay;
 class FOnlineAsyncTaskGooglePlayQueryInAppPurchasesV2;
-enum class EPurchaseTransactionState : uint8;
 
 /**
  * Implementation for online store via GooglePlay services
@@ -49,6 +48,9 @@ public:
 	/** Initialize the interface */
 	void Init();
 	
+	/** Handle Java side query available iap completed notification */
+	void OnGooglePlayAvailableIAPQueryComplete(EGooglePlayBillingResponseCode InResponse, const TArray<FProvidedProductInformation>& InProvidedProductInformation);
+
 private:
 	
 	void AddOffer(const TSharedRef<FOnlineStoreOffer>& NewOffer);
@@ -61,8 +63,8 @@ private:
 	/** Mapping of all queried offers to their product information */
 	FOnlineOfferDescriptionMap CachedOffers;
 
-	/** The current query for iap async task */
-	FOnlineAsyncTaskGooglePlayQueryInAppPurchasesV2* CurrentQueryTask;
+	/** The delegate to call for current query */
+	FOnQueryOnlineStoreOffersComplete QueryOnlineStoreOffersCompleteDelegate;
 
 	/** Is a query already in flight */
 	bool bIsQueryInFlight;
@@ -70,9 +72,6 @@ private:
 	/** Reference to the parent subsystem */
 	FOnlineSubsystemGooglePlay* Subsystem;
 
-	void OnGooglePlayAvailableIAPQueryComplete(EGooglePlayBillingResponseCode InResponse, const TArray<FProvidedProductInformation>& InProvidedProductInformation);
-	FDelegateHandle AvailableIAPQueryDelegateHandle;
-
 	/** Delegate fired when a query for purchases has completed, whether successful or unsuccessful */
 	FOnQueryForAvailablePurchasesComplete OnQueryForAvailablePurchasesCompleteDelegate;
 };
diff --git a/Engine/Plugins/Online/Android/OnlineSubsystemGooglePlay/Source/Public/OnlineSubsystemGooglePlay.h b/Engine/Plugins/Online/Android/OnlineSubsystemGooglePlay/Source/Public/OnlineSubsystemGooglePlay.h
--- a/Engine/Plugins/Online/Android/OnlineSubsystemGooglePlay/Source/Public/OnlineSubsystemGooglePlay.h
+++ b/Engine/Plugins/Online/Android/OnlineSubsystemGooglePlay/Source/Public/OnlineSubsystemGooglePlay.h
@@ -118,42 +118,6 @@ PACKAGE_SCOPE:
 	 */
 	bool IsInAppPurchasingEnabled();
 
-	/**
-	 * Delegate fired internally when the Java query for available in app purchases has completed, notifying any GooglePlay OSS listeners
-	 * Not meant for external use
-	 *
-	 * @param InResponseCode response from the GooglePlay backend
-	 * @param ProvidedProductInformation information returned from the backend about the queried offers
-	 */
-	DEFINE_ONLINE_DELEGATE_TWO_PARAM(OnGooglePlayAvailableIAPQueryComplete, EGooglePlayBillingResponseCode /*InResponseCode*/, const TArray<FProvidedProductInformation>& /*ProvidedProductInformation*/);
-
-	/**
-	 * Delegate fired internally when the Java in app purchase has completed, notifying any GooglePlay OSS listeners
-	 * Not meant for external use
-	 *
-	 * @param InResponseCode response from the GooglePlay backend
-	 * @param InTransactionData transaction information for the purchase
-	 */
-	DEFINE_ONLINE_DELEGATE_TWO_PARAM(OnGooglePlayProcessPurchaseComplete, EGooglePlayBillingResponseCode /*InResponseCode*/, const FGoogleTransactionData& /*InTransactionData*/);
-
-	/**
-	 * Delegate fired internally when the Java query for existing purchases has completed, notifying any GooglePlay OSS listeners
-	 * Not meant for external use
-	 *
-	 * @param InResponseCode response from the GooglePlay backend
-	 * @param InExistingPurchases known purchases for the user (non consumed or permanent)
-	 */
-	DEFINE_ONLINE_DELEGATE_TWO_PARAM(OnGooglePlayQueryExistingPurchasesComplete, EGooglePlayBillingResponseCode /*InResponseCode*/, const TArray<FGoogleTransactionData>& /*InExistingPurchases*/);
-
-	/**
-	 * Delegate fired internally when the Java query for restoring purchases has completed, notifying any GooglePlay OSS listeners
-	 * Not meant for external use
-	 *
-	 * @param InResponseCode response from the GooglePlay backend
-	 * @param InRestoredPurchases restored for the user (non consumed or permanent)
-	 */
-	DEFINE_ONLINE_DELEGATE_TWO_PARAM(OnGooglePlayRestorePurchasesComplete, EGooglePlayBillingResponseCode /*InResponseCode*/, const TArray<FGoogleTransactionData>& /*InRestoredPurchases*/);
-
 	/** Returns true if there are any async login tasks currently being tracked. */
 	bool AreAnyAsyncLoginTasksRunning() const { return CurrentLoginTask != nullptr || CurrentShowLoginUITask != nullptr; }
 
diff --git a/Engine/Plugins/Online/OnlineSubsystem/Source/Private/OnlineSubsystemImpl.cpp b/Engine/Plugins/Online/OnlineSubsystem/Source/Private/OnlineSubsystemImpl.cpp
--- a/Engine/Plugins/Online/OnlineSubsystem/Source/Private/OnlineSubsystemImpl.cpp
+++ b/Engine/Plugins/Online/OnlineSubsystem/Source/Private/OnlineSubsystemImpl.cpp
@@ -258,6 +258,10 @@ bool FOnlineSubsystemImpl::Exec(UWorld* InWorld, const TCHAR* Cmd, FOutputDevice
 	{
 		bWasHandled = HandlePurchaseExecCommands(InWorld, Cmd, Ar);
 	}
+	else if (FParse::Command(&Cmd, TEXT("STORE")))
+	{
+		bWasHandled = HandleStoreExecCommands(InWorld, Cmd, Ar);
+	}
 
 	return bWasHandled;
 }
@@ -283,9 +287,9 @@ void FOnlineSubsystemImpl::DumpReceipts(const FUniqueNetId& UserId)
 		{
 			for (const FPurchaseReceipt& Receipt : Receipts)
 			{
-				UE_LOG_ONLINE(Display, TEXT("Receipt: %s %d"),
+				UE_LOG_ONLINE(Display, TEXT("Receipt: %s %s"),
 							  *Receipt.TransactionId,
-							  (int32)Receipt.TransactionState);
+							  LexToString(Receipt.TransactionState));
 
 				UE_LOG_ONLINE(Display, TEXT("-Offers:"));
 				for (const FPurchaseReceipt::FReceiptOfferEntry& ReceiptOffer : Receipt.ReceiptOffers)
@@ -322,9 +326,9 @@ void FOnlineSubsystemImpl::FinalizeReceipts(const FUniqueNetId& UserId)
 		PurchaseInt->GetReceipts(UserId, Receipts);
 		for (const FPurchaseReceipt& Receipt : Receipts)
 		{
-			UE_LOG_ONLINE(Display, TEXT("Receipt: %s %d"),
+			UE_LOG_ONLINE(Display, TEXT("Receipt: %s %s"),
 				*Receipt.TransactionId,
-				(int32)Receipt.TransactionState);
+				LexToString(Receipt.TransactionState));
 			for (const FPurchaseReceipt::FReceiptOfferEntry& ReceiptOffer : Receipt.ReceiptOffers)
 			{
 				UE_LOG_ONLINE(Display, TEXT(" -Namespace: %s Id: %s Quantity: %d"),
@@ -364,11 +368,11 @@ bool FOnlineSubsystemImpl::HandlePurchaseExecCommands(UWorld* InWorld, const TCH
 {
 	bool bWasHandled = false;
 	
-	if (FParse::Command(&Cmd, TEXT("RECEIPTS")))
+	IOnlinePurchasePtr PurchaseInt = GetPurchaseInterface();
+	IOnlineIdentityPtr IdentityInt = GetIdentityInterface();
+	if (PurchaseInt.IsValid() && IdentityInt.IsValid())
 	{
-		IOnlinePurchasePtr PurchaseInt = GetPurchaseInterface();
-		IOnlineIdentityPtr IdentityInt = GetIdentityInterface();
-		if (PurchaseInt.IsValid() && IdentityInt.IsValid())
+		if (FParse::Command(&Cmd, TEXT("RECEIPTS")))
 		{
 			FString CommandStr = FParse::Token(Cmd, false);
 			if (CommandStr.IsEmpty())
@@ -409,11 +413,88 @@ bool FOnlineSubsystemImpl::HandlePurchaseExecCommands(UWorld* InWorld, const TCH
 				}
 			}
 		}
+		else if (FParse::Command(&Cmd, TEXT("BUY")))
+		{
+			FString ProductId = FParse::Token(Cmd, false);
+			FString UserIdStr = FParse::Token(Cmd, false);
+			if (ProductId.IsEmpty() || UserIdStr.IsEmpty())
+			{
+				UE_LOG_ONLINE(Warning, TEXT("usage: PURCHASE BUY <productid> <userid>"));
+			}
+			else
+			{
+				bWasHandled = true;
+				if (FUniqueNetIdPtr UserId = IdentityInt->CreateUniquePlayerId(UserIdStr))
+				{
+					FPurchaseCheckoutRequest Request;
+					Request.AccountId = UserIdStr;
+					Request.AddPurchaseOffer(FString(), ProductId, 1);
+					PurchaseInt->Checkout(*UserId, Request, FOnPurchaseCheckoutComplete::CreateLambda([](const FOnlineError& Result, const TSharedRef<FPurchaseReceipt>& Receipt)
+						{
+							UE_LOG_ONLINE(Log, TEXT("Checkout Result=%s TransactionId=%s TransactionState=%s"), *Result.ToLogString(), *Receipt->TransactionId, LexToString(Receipt->TransactionState));
+							for (const FPurchaseReceipt::FReceiptOfferEntry& OfferEntry : Receipt->ReceiptOffers)
+							{
+								UE_LOG_ONLINE(Log, TEXT("  OfferEntry Namespace=%s OfferId=%s Quantity=%d"), *OfferEntry.Namespace, *OfferEntry.OfferId, OfferEntry.Quantity);
+								for (const FPurchaseReceipt::FLineItemInfo& LineItem : OfferEntry.LineItems)
+								{
+									UE_LOG_ONLINE(Log, TEXT("    LineItem Name=%s UniqueId=%s ValidationInfo=%s"), *LineItem.ItemName, *LineItem.UniqueId, *LineItem.ValidationInfo);
+								}
+							}
+						}));
+				}
+			}
+		}
 	}
 
 	return bWasHandled;
 }
 
+bool FOnlineSubsystemImpl::HandleStoreExecCommands(UWorld* InWorld, const TCHAR* Cmd, FOutputDevice& Ar)
+{
+	bool bWasHandled = false;
+
+	IOnlineStoreV2Ptr StoreInt = GetStoreV2Interface();
+	IOnlineIdentityPtr IdentityInt = GetIdentityInterface();
+	if (StoreInt.IsValid() && IdentityInt.IsValid())
+	{
+		if (FParse::Command(&Cmd, TEXT("LIST")))
+		{
+			FString ProductStr = FParse::Token(Cmd, false);
+			FString UserIdStr = FParse::Token(Cmd, false);
+			if (ProductStr.IsEmpty() || UserIdStr.IsEmpty())
+			{
+				UE_LOG_ONLINE(Warning, TEXT("usage: STORE LIST <product> <userid>"));
+			}
+			else
+			{
+				bWasHandled = true;
+				if (FUniqueNetIdPtr UserId = IdentityInt->CreateUniquePlayerId(UserIdStr))
+				{
+					TArray<FString> ProductIds;
+					ProductIds.Add(ProductStr);
+					StoreInt->QueryOffersById(*UserId, ProductIds, FOnQueryOnlineStoreOffersComplete::CreateLambda([StoreInt](bool bWasSuccessful, const TArray<FUniqueOfferId>& /*OfferIds*/, const FString& Error) 
+					{
+						if(!bWasSuccessful)
+						{
+							UE_LOG_ONLINE(Error, TEXT("QueryOffersById failed: %s"), *Error);
+						}
+						TArray<FOnlineStoreOfferRef> Offers;
+						StoreInt->GetOffers(Offers);
+						for (const FOnlineStoreOfferRef& Product : Offers)
+						{
+							UE_LOG_ONLINE(Log, TEXT("Product Id=%s, Title=%s, Price=%s"),
+								*Product->OfferId,
+								*Product->Title.ToString(),
+								*Product->PriceText.ToString());
+						}
+					}));
+				}
+			}
+		}
+	}
+	return bWasHandled;
+}
+
 bool FOnlineSubsystemImpl::HandleFriendExecCommands(UWorld* InWorld, const TCHAR* Cmd, FOutputDevice& Ar)
 {
 	bool bWasHandled = false;
diff --git a/Engine/Plugins/Online/OnlineSubsystem/Source/Public/Interfaces/OnlinePurchaseInterface.h b/Engine/Plugins/Online/OnlineSubsystem/Source/Public/Interfaces/OnlinePurchaseInterface.h
--- a/Engine/Plugins/Online/OnlineSubsystem/Source/Public/Interfaces/OnlinePurchaseInterface.h
+++ b/Engine/Plugins/Online/OnlineSubsystem/Source/Public/Interfaces/OnlinePurchaseInterface.h
@@ -176,6 +176,11 @@ public:
 		ReceiptOffers.Add(ReceiptOffer);
 	}
 
+	void AddReceiptOffer(FReceiptOfferEntry&& ReceiptOffer)
+	{
+		ReceiptOffers.Add(MoveTemp(ReceiptOffer));
+	}
+
 	BEGIN_ONLINE_JSON_SERIALIZER
 		ONLINE_JSON_SERIALIZE("transactionId", TransactionId);
 		ONLINE_JSON_SERIALIZE_ARRAY_SERIALIZABLE("receiptList", ReceiptOffers, FReceiptOfferEntry);
diff --git a/Engine/Plugins/Online/OnlineSubsystem/Source/Public/OnlineSubsystemImpl.h b/Engine/Plugins/Online/OnlineSubsystem/Source/Public/OnlineSubsystemImpl.h
--- a/Engine/Plugins/Online/OnlineSubsystem/Source/Public/OnlineSubsystemImpl.h
+++ b/Engine/Plugins/Online/OnlineSubsystem/Source/Public/OnlineSubsystemImpl.h
@@ -34,6 +34,7 @@ private:
 	bool HandleSessionExecCommands(UWorld* InWorld, const TCHAR* Cmd, FOutputDevice& Ar);
 	bool HandlePresenceExecCommands(UWorld* InWorld, const TCHAR* Cmd, FOutputDevice& Ar);
 	bool HandlePurchaseExecCommands(UWorld* InWorld, const TCHAR* Cmd, FOutputDevice& Ar);
+	bool HandleStoreExecCommands(UWorld* InWorld, const TCHAR* Cmd, FOutputDevice& Ar);
 	
 	/** Delegate fired when exec cheat related to receipts completes */
 	void OnQueryReceiptsComplete(const FOnlineError& Result, FUniqueNetIdPtr UserId);
diff --git a/Engine/Source/Programs/UnrealBuildTool/Platform/Android/UEDeployAndroid.cs b/Engine/Source/Programs/UnrealBuildTool/Platform/Android/UEDeployAndroid.cs
--- a/Engine/Source/Programs/UnrealBuildTool/Platform/Android/UEDeployAndroid.cs
+++ b/Engine/Source/Programs/UnrealBuildTool/Platform/Android/UEDeployAndroid.cs
@@ -5601,7 +5601,6 @@ namespace UnrealBuildTool
 				{ "//$${gameActivityOnResumeAdditions}$$", UPL.ProcessPluginNode(NDKArch, "gameActivityOnResumeAdditions", "")},
 				{ "//$${gameActivityOnNewIntentAdditions}$$", UPL.ProcessPluginNode(NDKArch, "gameActivityOnNewIntentAdditions", "")},
   				{ "//$${gameActivityOnActivityResultAdditions}$$", UPL.ProcessPluginNode(NDKArch, "gameActivityOnActivityResultAdditions", "")},
-				{ "//$${gameActivityOnActivityResultIapStoreHelperHandler}$$", UPL.ProcessPluginNode(NDKArch, "gameActivityOnActivityResultIapStoreHelperHandler", "")},
   				{ "//$${gameActivityPreConfigRulesParseAdditions}$$", UPL.ProcessPluginNode(NDKArch, "gameActivityPreConfigRulesParseAdditions", "")},
   				{ "//$${gameActivityPostConfigRulesAdditions}$$", UPL.ProcessPluginNode(NDKArch, "gameActivityPostConfigRulesAdditions", "")},
   				{ "//$${gameActivityFinalizeConfigRulesAdditions}$$", UPL.ProcessPluginNode(NDKArch, "gameActivityFinalizeConfigRulesAdditions", "")},
diff --git a/Engine/Source/Runtime/Launch/Private/Android/AndroidJNI.cpp b/Engine/Source/Runtime/Launch/Private/Android/AndroidJNI.cpp
--- a/Engine/Source/Runtime/Launch/Private/Android/AndroidJNI.cpp
+++ b/Engine/Source/Runtime/Launch/Private/Android/AndroidJNI.cpp
@@ -214,7 +214,6 @@ void FJavaWrapper::FindGooglePlayBillingMethods(JNIEnv* Env)
 	AndroidThunkJava_IapQueryInAppPurchases = FindMethod(Env, GoogleServicesClassID, "AndroidThunkJava_IapQueryInAppPurchases", "([Ljava/lang/String;)Z", bIsStoreOptional);
 	AndroidThunkJava_IapBeginPurchase = FindMethod(Env, GoogleServicesClassID, "AndroidThunkJava_IapBeginPurchase", "(Ljava/lang/String;Ljava/lang/String;)Z", bIsStoreOptional);
 	AndroidThunkJava_IapIsAllowedToMakePurchases = FindMethod(Env, GoogleServicesClassID, "AndroidThunkJava_IapIsAllowedToMakePurchases", "()Z", bIsStoreOptional);
-	AndroidThunkJava_IapRestorePurchases = FindMethod(Env, GoogleServicesClassID, "AndroidThunkJava_IapRestorePurchases", "([Ljava/lang/String;[Z)Z", bIsStoreOptional);
 	AndroidThunkJava_IapConsumePurchase = FindMethod(Env, GoogleServicesClassID, "AndroidThunkJava_IapConsumePurchase", "(Ljava/lang/String;)Z", bIsStoreOptional);
 	AndroidThunkJava_IapQueryExistingPurchases = FindMethod(Env, GoogleServicesClassID, "AndroidThunkJava_IapQueryExistingPurchases", "()Z", bIsStoreOptional);
 }
@@ -445,7 +444,6 @@ jmethodID FJavaWrapper::AndroidThunkJava_IapSetupService;
 jmethodID FJavaWrapper::AndroidThunkJava_IapQueryInAppPurchases;
 jmethodID FJavaWrapper::AndroidThunkJava_IapBeginPurchase;
 jmethodID FJavaWrapper::AndroidThunkJava_IapIsAllowedToMakePurchases;
-jmethodID FJavaWrapper::AndroidThunkJava_IapRestorePurchases;
 jmethodID FJavaWrapper::AndroidThunkJava_IapQueryExistingPurchases;
 jmethodID FJavaWrapper::AndroidThunkJava_IapConsumePurchase;
 
@@ -1442,36 +1440,6 @@ bool AndroidThunkCpp_Iap_IsAllowedToMakePurchases()
 	return bResult;
 }
 
-bool AndroidThunkCpp_Iap_RestorePurchases(const TArray<FString>& ProductIDs, const TArray<bool>& bConsumable)
-{
-	FPlatformMisc::LowLevelOutputDebugString(TEXT("[JNI] - AndroidThunkCpp_Iap_RestorePurchases"));
-	bool bResult = false;
-
-	if (JNIEnv* Env = FAndroidApplication::GetJavaEnv())
-	{
-		CHECK_JNI_METHOD(FJavaWrapper::AndroidThunkJava_IapRestorePurchases);
-
-		// Populate some java types with the provided product information
-		auto ProductIDArray = NewScopedJavaObject(Env, (jobjectArray)Env->NewObjectArray(ProductIDs.Num(), FJavaWrapper::JavaStringClass, NULL));
-		auto ConsumeArray = NewScopedJavaObject(Env, (jbooleanArray)Env->NewBooleanArray(ProductIDs.Num()));
-
-		jboolean* ConsumeArrayValues = Env->GetBooleanArrayElements(*ConsumeArray, 0);
-		for (uint32 Param = 0; Param < ProductIDs.Num(); Param++)
-		{
-			auto StringValue = FJavaHelper::ToJavaString(Env, ProductIDs[Param]);
-			Env->SetObjectArrayElement(*ProductIDArray, Param, *StringValue);
-
-			ConsumeArrayValues[Param] = bConsumable[Param];
-		}
-		Env->ReleaseBooleanArrayElements(*ConsumeArray, ConsumeArrayValues, 0);
-
-		// Execute the java code for this operation
-		bResult = FJavaWrapper::CallBooleanMethod(Env, FJavaWrapper::GoogleServicesThis, FJavaWrapper::AndroidThunkJava_IapRestorePurchases, *ProductIDArray, *ConsumeArray);
-	}
-
-	return bResult;
-}
-
 void AndroidThunkCpp_UseSurfaceViewWorkaround()
 {
 	if (JNIEnv* Env = FAndroidApplication::GetJavaEnv())
diff --git a/Engine/Source/Runtime/Launch/Public/Android/AndroidJNI.h b/Engine/Source/Runtime/Launch/Public/Android/AndroidJNI.h
--- a/Engine/Source/Runtime/Launch/Public/Android/AndroidJNI.h
+++ b/Engine/Source/Runtime/Launch/Public/Android/AndroidJNI.h
@@ -115,7 +115,6 @@ public:
 	static jmethodID AndroidThunkJava_IapQueryInAppPurchases;
 	static jmethodID AndroidThunkJava_IapBeginPurchase;
 	static jmethodID AndroidThunkJava_IapIsAllowedToMakePurchases;
-	static jmethodID AndroidThunkJava_IapRestorePurchases;
 	static jmethodID AndroidThunkJava_IapQueryExistingPurchases;
 	static jmethodID AndroidThunkJava_IapConsumePurchase;
 
