Fix crash on OnlineStoreGooglePlay for Android in 32 bit architectures reported by user because of wrong printf format specifier used in log
Also fixed other log reporting for same struct field

diff --git a/Engine/Plugins/Online/Android/OnlineSubsystemGooglePlay/Source/Private/OnlineStoreGooglePlay.cpp b/Engine/Plugins/Online/Android/OnlineSubsystemGooglePlay/Source/Private/OnlineStoreGooglePlay.cpp
--- a/Engine/Plugins/Online/Android/OnlineSubsystemGooglePlay/Source/Private/OnlineStoreGooglePlay.cpp
+++ b/Engine/Plugins/Online/Android/OnlineSubsystemGooglePlay/Source/Private/OnlineStoreGooglePlay.cpp
@@ -69,7 +69,7 @@ void FOnlineStoreGooglePlayV2::OnGooglePlayAvailableIAPQueryComplete(EGooglePlay
 			AddOffer(NewProductOffer);
 			OfferIds.Add(NewProductOffer->OfferId);
 
-			UE_LOG_ONLINE_STOREV2(Log, TEXT("Product Identifier: %s, Name: %s, Desc: %s, Long Desc: %s, Price: %s IntPrice: %d"),
+			UE_LOG_ONLINE_STOREV2(Log, TEXT("Product Identifier: %s, Name: %s, Desc: %s, Long Desc: %s, Price: %s IntPrice: %lld"),
 				*NewProductOffer->OfferId,
 				*NewProductOffer->Title.ToString(),
 				*NewProductOffer->Description.ToString(),
@@ -212,7 +212,7 @@ JNI_METHOD void Java_com_epicgames_unreal_GooglePlayStoreHelper_NativeQueryCompl
 
 			ProvidedProductInformation.Add(NewProductInfo);
 
-			FPlatformMisc::LowLevelOutputDebugStringf(TEXT("\nProduct Identifier: %s, Name: %s, Description: %s, Price: %s, Price Raw: %d, Currency Code: %s\n"),
+			UE_LOG_ONLINE_STOREV2(Verbose, TEXT("\nProduct Identifier: %s, Name: %s, Description: %s, Price: %s, Price Raw: %lld, Currency Code: %s\n"),
 				*NewProductInfo.OfferId,
 				*NewProductInfo.Title.ToString(),
 				*NewProductInfo.Description.ToString(),
diff --git a/Engine/Plugins/Online/OnlineSubsystemUtils/Source/OnlineSubsystemUtils/Private/Tests/TestStoreInterface.cpp b/Engine/Plugins/Online/OnlineSubsystemUtils/Source/OnlineSubsystemUtils/Private/Tests/TestStoreInterface.cpp
--- a/Engine/Plugins/Online/OnlineSubsystemUtils/Source/OnlineSubsystemUtils/Private/Tests/TestStoreInterface.cpp
+++ b/Engine/Plugins/Online/OnlineSubsystemUtils/Source/OnlineSubsystemUtils/Private/Tests/TestStoreInterface.cpp
@@ -90,7 +90,7 @@ void FTestStoreInterface::OnQueryOnlineStoreOffersComplete(bool bWasSuccessful,
 		UE_LOG(LogOnline, Log, TEXT("Found %d offers in cache"), StoreOffers.Num());
 		for (const FOnlineStoreOfferRef& OfferRef : StoreOffers)
 		{
-			UE_LOG(LogOnline, Log, TEXT("  Title=[%s] Desc=[%s] OfferId=[%s] CurrencyCode=[%s] PriceInt=[%d] DisplayPrice=[%s]"), 
+			UE_LOG(LogOnline, Log, TEXT("  Title=[%s] Desc=[%s] OfferId=[%s] CurrencyCode=[%s] PriceInt=[%lld] DisplayPrice=[%s]"), 
 				*OfferRef->Title.ToString(),
 				*OfferRef->Description.ToString(),
 				*OfferRef->OfferId, 
