Fix possible loss of precission when passing currency values through JNI for GooglePlay Billing Library

diff --git a/Engine/Plugins/Online/Android/OnlineSubsystemGooglePlay/Source/Java/BillingApiV2/com/epicgames/unreal/GooglePlayStoreHelper.java b/Engine/Plugins/Online/Android/OnlineSubsystemGooglePlay/Source/Java/BillingApiV2/com/epicgames/unreal/GooglePlayStoreHelper.java
--- a/Engine/Plugins/Online/Android/OnlineSubsystemGooglePlay/Source/Java/BillingApiV2/com/epicgames/unreal/GooglePlayStoreHelper.java
+++ b/Engine/Plugins/Online/Android/OnlineSubsystemGooglePlay/Source/Java/BillingApiV2/com/epicgames/unreal/GooglePlayStoreHelper.java
@@ -273,7 +273,7 @@ public class GooglePlayStoreHelper implements StoreHelper, PurchasesUpdatedListe
 		public String[] Titles;
 		public String[] Descriptions;
 		public String[] Prices;
-		public float[] PricesRaw;
+		public long[] PricesRaw;
 		public String[] CurrencyCodes;
 
 		public ListOfProductDetailsAsArrays(List<ProductDetails> InProductDetailsList)
@@ -282,7 +282,7 @@ public class GooglePlayStoreHelper implements StoreHelper, PurchasesUpdatedListe
 			Titles = new String[InProductDetailsList.size()];
 			Descriptions = new String[InProductDetailsList.size()];
 			Prices = new String[InProductDetailsList.size()];
-			PricesRaw = new float[InProductDetailsList.size()];
+			PricesRaw = new long[InProductDetailsList.size()];
 			CurrencyCodes = new String[InProductDetailsList.size()];
 
 			for (int index = 0; index < InProductDetailsList.size(); index++) 
@@ -301,7 +301,7 @@ public class GooglePlayStoreHelper implements StoreHelper, PurchasesUpdatedListe
 				Titles[index] = ThisProduct.getTitle();
 				Descriptions[index] = ThisProduct.getDescription();
 				Prices[index] = PriceDetails.getFormattedPrice();
-				PricesRaw[index] = (float)(PriceDetails.getPriceAmountMicros() / 1000000.0);
+				PricesRaw[index] = PriceDetails.getPriceAmountMicros();
 				CurrencyCodes[index] = PriceDetails.getPriceCurrencyCode();
 			}
 		}
@@ -459,7 +459,7 @@ public class GooglePlayStoreHelper implements StoreHelper, PurchasesUpdatedListe
 	}
 
 	// Callback that notify the C++ implementation that a task has completed
-	public native void NativeQueryComplete(int ResponseCode, String[] ProductIDs, String[] Titles, String[] Descriptions, String[] Prices, float[] PricesRaw, String[] CurrencyCodes);
+	public native void NativeQueryComplete(int ResponseCode, String[] ProductIDs, String[] Titles, String[] Descriptions, String[] Prices, long[] PricesRaw, String[] CurrencyCodes);
 	public native void NativePurchaseComplete(int ResponseCode, String ProductId, int PurchaseState, String ProductToken, String ReceiptData, String Signature);
 	public native void NativeQueryExistingPurchasesComplete(int ResponseCode, String[] ProductIds, int[] PurchaseState, String[] ProductTokens, String[] ReceiptsData, String[] Signatures);
 	public native void NativeConsumeComplete(int ResponseCode, String PurchaseToken);
diff --git a/Engine/Plugins/Online/Android/OnlineSubsystemGooglePlay/Source/Private/OnlineStoreGooglePlay.cpp b/Engine/Plugins/Online/Android/OnlineSubsystemGooglePlay/Source/Private/OnlineStoreGooglePlay.cpp
--- a/Engine/Plugins/Online/Android/OnlineSubsystemGooglePlay/Source/Private/OnlineStoreGooglePlay.cpp
+++ b/Engine/Plugins/Online/Android/OnlineSubsystemGooglePlay/Source/Private/OnlineStoreGooglePlay.cpp
@@ -160,7 +160,7 @@ TSharedPtr<FOnlineStoreOffer> FOnlineStoreGooglePlayV2::GetOffer(const FUniqueOf
 	return Result;
 }
 
-JNI_METHOD void Java_com_epicgames_unreal_GooglePlayStoreHelper_NativeQueryComplete(JNIEnv* jenv, jobject /*Thiz*/, jint ResponseCode, jobjectArray ProductIDs, jobjectArray Titles, jobjectArray Descriptions, jobjectArray Prices, jfloatArray PriceValuesRaw, jobjectArray CurrencyCodes)
+JNI_METHOD void Java_com_epicgames_unreal_GooglePlayStoreHelper_NativeQueryComplete(JNIEnv* jenv, jobject /*Thiz*/, jint ResponseCode, jobjectArray ProductIDs, jobjectArray Titles, jobjectArray Descriptions, jobjectArray Prices, jlongArray PriceValuesRaw, jobjectArray CurrencyCodes)
 {
 	TArray<FOnlineStoreOffer> ProvidedProductInformation;
 	EGooglePlayBillingResponseCode EGPResponse = (EGooglePlayBillingResponseCode)ResponseCode;
@@ -177,7 +177,7 @@ JNI_METHOD void Java_com_epicgames_unreal_GooglePlayStoreHelper_NativeQueryCompl
 
 		ensure((NumProducts == NumTitles) && (NumProducts == NumDescriptions) && (NumProducts == NumPrices) && (NumProducts == NumPricesRaw) && (NumProducts == NumCurrencyCodes));
 
-		jfloat* PriceValues = jenv->GetFloatArrayElements(PriceValuesRaw, 0);
+		jlong* PriceValues = jenv->GetLongArrayElements(PriceValuesRaw, 0);
 
 		for (jsize Idx = 0; Idx < NumProducts; Idx++)
 		{
@@ -204,7 +204,10 @@ JNI_METHOD void Java_com_epicgames_unreal_GooglePlayStoreHelper_NativeQueryCompl
 
 			const FDecimalNumberFormattingRules& FormattingRules = Culture.GetCurrencyFormattingRules(NewProductInfo.CurrencyCode);
 			const FNumberFormattingOptions& FormattingOptions = FormattingRules.CultureDefaultFormattingOptions;
-			double Val = static_cast<double>(PriceValues[Idx]) * static_cast<double>(FMath::Pow(10.0f, FormattingOptions.MaximumFractionalDigits));
+
+			// Value provided in PriceValues[idx] by Google Billing Library is multiplied by 1'000'000 (it is a value in micro units). We need to divide it by 10^6 and multiply by 
+			// 10^FormattingOptions.MaximumFractionalDigits to properly adjust the value to the number of fractional digits expected
+			double Val = static_cast<double>(PriceValues[Idx]) * static_cast<double>(FMath::Pow(10.0f, FormattingOptions.MaximumFractionalDigits - 6));
 
 			NewProductInfo.NumericPrice = FMath::TruncToInt(Val + 0.5);
 			NewProductInfo.ReleaseDate = FDateTime::MinValue();
@@ -220,7 +223,7 @@ JNI_METHOD void Java_com_epicgames_unreal_GooglePlayStoreHelper_NativeQueryCompl
 				NewProductInfo.NumericPrice,
 				*NewProductInfo.CurrencyCode);
 		}
-		jenv->ReleaseFloatArrayElements(PriceValuesRaw, PriceValues, JNI_ABORT);
+		jenv->ReleaseLongArrayElements(PriceValuesRaw, PriceValues, JNI_ABORT);
 	}
 
 	UE_LOG_ONLINE_STOREV2(Verbose, TEXT("QueryOffersById result Success: %d Response: %s"), bWasSuccessful, LexToString(EGPResponse));
