Allow purchasing multiple products in one transaction on GooglePlay Android

diff --git a/Engine/Build/Android/Java/src/com/epicgames/unreal/StoreHelper.java b/Engine/Build/Android/Java/src/com/epicgames/unreal/StoreHelper.java
--- a/Engine/Build/Android/Java/src/com/epicgames/unreal/StoreHelper.java
+++ b/Engine/Build/Android/Java/src/com/epicgames/unreal/StoreHelper.java
@@ -8,7 +8,7 @@ import android.os.Bundle;
 public interface StoreHelper
 {
 	public boolean QueryInAppPurchases(String[] ProductIDs);
-	public boolean BeginPurchase(String ProductID, String ObfuscatedAccountId);
+	public boolean BeginPurchase(String[] ProductIDs, String ObfuscatedAccountId);
 	public boolean IsAllowedToMakePurchases();
 	public void ConsumePurchase(String purchaseToken);
 	public boolean QueryExistingPurchases();
diff --git a/Engine/Plugins/Online/Android/OnlineSubsystemGooglePlay/Source/Java/BillingApiV2/com/epicgames/unreal/GooglePlayStoreHelper.java b/Engine/Plugins/Online/Android/OnlineSubsystemGooglePlay/Source/Java/BillingApiV2/com/epicgames/unreal/GooglePlayStoreHelper.java
--- a/Engine/Plugins/Online/Android/OnlineSubsystemGooglePlay/Source/Java/BillingApiV2/com/epicgames/unreal/GooglePlayStoreHelper.java
+++ b/Engine/Plugins/Online/Android/OnlineSubsystemGooglePlay/Source/Java/BillingApiV2/com/epicgames/unreal/GooglePlayStoreHelper.java
@@ -23,6 +23,7 @@ import com.android.vending.billing.util.Base64; //WMM should use different base6
 import java.util.ArrayList;
 import java.util.Arrays;
 import java.util.List;
+import java.util.stream.Collectors;
 
 public class GooglePlayStoreHelper implements StoreHelper, PurchasesUpdatedListener
 {
@@ -141,46 +142,38 @@ public class GooglePlayStoreHelper implements StoreHelper, PurchasesUpdatedListe
 	 * Start the purchase flow for a particular product id
 	 */
 	@Override
-	public boolean BeginPurchase(final String ProductId, final String ObfuscatedAccountId)
+	public boolean BeginPurchase(final String[] ProductIds, final String ObfuscatedAccountId)
 	{
-		QueryProductDetailsParams Params = CreateQueryProductDetailsParamsForProduct(ProductId);
+		QueryProductDetailsParams Params = CreateQueryProductDetailsParamsForProducts(ProductIds);
 
+		Log.debug("[GooglePlayStoreHelper] - GooglePlayStoreHelper::BeginPurchase - querying product information " + String.join(",", ProductIds));
 		Client.queryProductDetailsAsync(Params, new ProductDetailsResponseListener() {
 			@Override
 			public void onProductDetailsResponse(@NonNull BillingResult Result, @NonNull List<ProductDetails> ProductDetailsList)
 			{
 				if(Result.getResponseCode() == BillingClient.BillingResponseCode.OK)
 				{
-					final ProductDetails DetailsForProductId = FindProductDetailsInList(ProductId, ProductDetailsList);
-					if(DetailsForProductId != null)
+					// BillingClient.launchBillingFlow is annotated @UiThread, so dispatch the call to the Activity UI thread
+					MainActivity.runOnUiThread(new Runnable()
 					{
-						// BillingClient.launchBillingFlow is annotated @UiThread, so dispatch the call to the Activity UI thread
-						MainActivity.runOnUiThread(new Runnable()
-						{
-							@Override
-							public void run() {
-								Log.debug("[GooglePlayStoreHelper] - GooglePlayStoreHelper::BeginPurchase - Launching billing flow " + ProductId);
-								BillingFlowParams Params = CreateBillingFlowParams(DetailsForProductId, ObfuscatedAccountId);
-								BillingResult Result = Client.launchBillingFlow(MainActivity, Params);
-								int ResponseCode = Result.getResponseCode();
-								if (ResponseCode != BillingClient.BillingResponseCode.OK) 
-								{
-									Log.debug("[GooglePlayStoreHelper] - GooglePlayStoreHelper::BeginPurchase - Failed! " + TranslateBillingResult(Result));
-									NativePurchaseComplete(ResponseCode, ProductId, Purchase.PurchaseState.UNSPECIFIED_STATE, "", "", "");
-								}
+						@Override
+						public void run() {
+							Log.debug("[GooglePlayStoreHelper] - GooglePlayStoreHelper::BeginPurchase - Launching billing flow " + ProductDetailsList.stream().map(ProductDetails::getName).collect(Collectors.joining(",")));
+							BillingFlowParams Params = CreateBillingFlowParams(ProductDetailsList, ObfuscatedAccountId);
+							BillingResult Result = Client.launchBillingFlow(MainActivity, Params);
+							int ResponseCode = Result.getResponseCode();
+							if (ResponseCode != BillingClient.BillingResponseCode.OK) 
+							{
+								Log.debug("[GooglePlayStoreHelper] - GooglePlayStoreHelper::BeginPurchase - Failed! " + TranslateBillingResult(Result));
+								NativePurchaseComplete(ResponseCode, ProductIds, Purchase.PurchaseState.UNSPECIFIED_STATE, "", "", "");
 							}
-						});
-					}
-					else
-					{
-						Log.debug("[GooglePlayStoreHelper] - GooglePlayStoreHelper::BeginPurchase - Failed! Could not find ProductDetails for " + ProductId);
-						NativePurchaseComplete(CustomLogicErrorResponse, ProductId, Purchase.PurchaseState.UNSPECIFIED_STATE,"", "", "");
-					}
+						}
+					});
 				}
 				else
 				{
 					Log.debug("[GooglePlayStoreHelper] - GooglePlayStoreHelper::BeginPurchase - Failed! " + TranslateBillingResult(Result));
-					NativePurchaseComplete(Result.getResponseCode(), ProductId, Purchase.PurchaseState.UNSPECIFIED_STATE,"", "", "");
+					NativePurchaseComplete(Result.getResponseCode(), ProductIds, Purchase.PurchaseState.UNSPECIFIED_STATE,"", "", "");
 				}
 			}
 		});
@@ -248,19 +241,19 @@ public class GooglePlayStoreHelper implements StoreHelper, PurchasesUpdatedListe
 		{
 			for (final Purchase UpdatedPurchase : UpdatedPurchases) 
 			{
-				final String ProductId = UpdatedPurchase.getProducts().get(0);
-				Log.debug("[GooglePlayStoreHelper] - GooglePlayStoreHelper::onPurchasesUpdated - Processing purchase of product " + ProductId);
+				final String[] ProductIds = UpdatedPurchase.getProducts().toArray(new String[0]);
+				Log.debug("[GooglePlayStoreHelper] - GooglePlayStoreHelper::onPurchasesUpdated - Processing purchase of product " + String.join(",",ProductIds));
 				Log.debug("[GooglePlayStoreHelper] - Purchase data: " + UpdatedPurchase.toString());
 				Log.debug("[GooglePlayStoreHelper] - Data signature: " + UpdatedPurchase.getSignature());
 
 				String Receipt = Base64.encode(UpdatedPurchase.getOriginalJson().getBytes());
-				NativePurchaseComplete(BillingClient.BillingResponseCode.OK, ProductId, UpdatedPurchase.getPurchaseState(), UpdatedPurchase.getPurchaseToken(), Receipt, UpdatedPurchase.getSignature());
+				NativePurchaseComplete(BillingClient.BillingResponseCode.OK, ProductIds, UpdatedPurchase.getPurchaseState(), UpdatedPurchase.getPurchaseToken(), Receipt, UpdatedPurchase.getSignature());
 			}
 		}
 		else
 		{
 			Log.debug("[GooglePlayStoreHelper] - GooglePlayStoreHelper::onPurchasesUpdated - Purchase failed. Response code: " + TranslateBillingResult(Result));
- 			NativePurchaseComplete(responseCode, "", Purchase.PurchaseState.UNSPECIFIED_STATE, "", "", "");
+ 			NativePurchaseComplete(responseCode, null, Purchase.PurchaseState.UNSPECIFIED_STATE, "", "", "");
 		}
 	}
 
@@ -309,7 +302,7 @@ public class GooglePlayStoreHelper implements StoreHelper, PurchasesUpdatedListe
 
 	class ListOfPurchasesAsArrays 
 	{
-		public String[] OwnedProducts;
+		public String[][] OwnedProducts;
 		public int[] PurchaseStates;
 		public String[] ProductTokens;
 		public String[] Signatures;
@@ -317,7 +310,7 @@ public class GooglePlayStoreHelper implements StoreHelper, PurchasesUpdatedListe
 
 		public ListOfPurchasesAsArrays(List<Purchase> InOwnedItems) 
 		{
-			OwnedProducts = new String[InOwnedItems.size()];
+			OwnedProducts = new String[InOwnedItems.size()][];
 			PurchaseStates = new int[InOwnedItems.size()];
 			ProductTokens = new String[InOwnedItems.size()];
 			Signatures = new String[InOwnedItems.size()];
@@ -326,7 +319,7 @@ public class GooglePlayStoreHelper implements StoreHelper, PurchasesUpdatedListe
 			for (int index = 0; index < InOwnedItems.size(); index ++) 
 			{
 				Purchase OwnedPurchase = InOwnedItems.get(index);
-				OwnedProducts[index] = OwnedPurchase.getProducts().get(0);
+				OwnedProducts[index] = OwnedPurchase.getProducts().toArray(new String[0]);
 				PurchaseStates[index] = OwnedPurchase.getPurchaseState();
 				ProductTokens[index] = OwnedPurchase.getPurchaseToken();
 				Signatures[index] = OwnedPurchase.getSignature();
@@ -335,13 +328,15 @@ public class GooglePlayStoreHelper implements StoreHelper, PurchasesUpdatedListe
 		}
 	}
 	
-	private static BillingFlowParams CreateBillingFlowParams(ProductDetails InProductDetails, String InObfuscatedAccountId)
+	private static BillingFlowParams CreateBillingFlowParams(List<ProductDetails> InProductDetails, String InObfuscatedAccountId)
 	{
-		ArrayList<BillingFlowParams.ProductDetailsParams> ProductDetailsParamList = new ArrayList<>(1);
-		ProductDetailsParamList.add(BillingFlowParams.ProductDetailsParams.newBuilder()
-			.setProductDetails(InProductDetails)
-			.build());
-		
+		ArrayList<BillingFlowParams.ProductDetailsParams> ProductDetailsParamList = new ArrayList<>(InProductDetails.size());
+		for(ProductDetails Details : InProductDetails) 
+		{
+			ProductDetailsParamList.add(BillingFlowParams.ProductDetailsParams.newBuilder()
+				.setProductDetails(Details)
+				.build());
+		}
 		BillingFlowParams.Builder Params = BillingFlowParams.newBuilder()
 			.setProductDetailsParamsList(ProductDetailsParamList);
 		if (InObfuscatedAccountId != null)
@@ -351,18 +346,6 @@ public class GooglePlayStoreHelper implements StoreHelper, PurchasesUpdatedListe
 		return Params.build();
 	}
 
-	private static QueryProductDetailsParams CreateQueryProductDetailsParamsForProduct(String InProductId)
-	{
-		ArrayList<QueryProductDetailsParams.Product> ProductList = new ArrayList<>(1);
-		ProductList.add(QueryProductDetailsParams.Product.newBuilder()
-			.setProductId(InProductId)
-			.setProductType(BillingClient.ProductType.INAPP)
-			.build());
-		return QueryProductDetailsParams.newBuilder()
-			.setProductList(ProductList)
-			.build();
-	}
-
 	private static QueryProductDetailsParams CreateQueryProductDetailsParamsForProducts(String[] InProductIds)
 	{
 		ArrayList<QueryProductDetailsParams.Product> ProductList = new ArrayList<>();
@@ -380,18 +363,6 @@ public class GooglePlayStoreHelper implements StoreHelper, PurchasesUpdatedListe
 			.build();
 	}
 
-	private static ProductDetails FindProductDetailsInList(String ProductId, List<ProductDetails> InProductDetailsList)
-	{
-		for(ProductDetails Details : InProductDetailsList) 
-		{
-			if (Details.getProductId().equals(ProductId)) 
-			{
-				return Details;
-			}
-		}
-		return null;
-	}
-	
 	/**
 	 * Get a text translation of the Response Codes returned by google play.
 	 */
@@ -459,8 +430,8 @@ public class GooglePlayStoreHelper implements StoreHelper, PurchasesUpdatedListe
 	}
 
 	// Callback that notify the C++ implementation that a task has completed
-	public native void NativeQueryComplete(int ResponseCode, String[] ProductIDs, String[] Titles, String[] Descriptions, String[] Prices, long[] PricesRaw, String[] CurrencyCodes);
-	public native void NativePurchaseComplete(int ResponseCode, String ProductId, int PurchaseState, String ProductToken, String ReceiptData, String Signature);
-	public native void NativeQueryExistingPurchasesComplete(int ResponseCode, String[] ProductIds, int[] PurchaseState, String[] ProductTokens, String[] ReceiptsData, String[] Signatures);
+	public native void NativeQueryComplete(int ResponseCode, String[] ProductIds, String[] Titles, String[] Descriptions, String[] Prices, long[] PricesRaw, String[] CurrencyCodes);
+	public native void NativePurchaseComplete(int ResponseCode, String[] ProductIds, int PurchaseState, String ProductToken, String ReceiptData, String Signature);
+	public native void NativeQueryExistingPurchasesComplete(int ResponseCode, String[][] ProductIds, int[] PurchaseState, String[] ProductTokens, String[] ReceiptsData, String[] Signatures);
 	public native void NativeConsumeComplete(int ResponseCode, String PurchaseToken);
 }
diff --git a/Engine/Plugins/Online/Android/OnlineSubsystemGooglePlay/Source/OnlineSubsystemGooglePlay_UPL.xml b/Engine/Plugins/Online/Android/OnlineSubsystemGooglePlay/Source/OnlineSubsystemGooglePlay_UPL.xml
--- a/Engine/Plugins/Online/Android/OnlineSubsystemGooglePlay/Source/OnlineSubsystemGooglePlay_UPL.xml
+++ b/Engine/Plugins/Online/Android/OnlineSubsystemGooglePlay/Source/OnlineSubsystemGooglePlay_UPL.xml
@@ -22,7 +22,7 @@
 
 	<gameActivityAndroidThunkJavaIapBeginPurchase>
 		<insert>
-public boolean AndroidThunkJava_IapBeginPurchase(String ProductId, String AccountId)
+public boolean AndroidThunkJava_IapBeginPurchase(String[] ProductIds, String AccountId)
 {
 	Log.debug("[JAVA] - AndroidThunkJava_IapBeginPurchase");
 	boolean bTriggeredPurchase = false;
@@ -50,7 +50,7 @@ public boolean AndroidThunkJava_IapBeginPurchase(String ProductId, String Accoun
 			{
 			}
 		}
-		bTriggeredPurchase = IapStoreHelper.BeginPurchase(ProductId, ObfuscatedAccountId);
+		bTriggeredPurchase = IapStoreHelper.BeginPurchase(ProductIds, ObfuscatedAccountId);
 	}
 	else
 	{
diff --git a/Engine/Plugins/Online/Android/OnlineSubsystemGooglePlay/Source/Private/OnlinePurchaseGooglePlay.cpp b/Engine/Plugins/Online/Android/OnlineSubsystemGooglePlay/Source/Private/OnlinePurchaseGooglePlay.cpp
--- a/Engine/Plugins/Online/Android/OnlineSubsystemGooglePlay/Source/Private/OnlinePurchaseGooglePlay.cpp
+++ b/Engine/Plugins/Online/Android/OnlineSubsystemGooglePlay/Source/Private/OnlinePurchaseGooglePlay.cpp
@@ -4,12 +4,13 @@
 #include "OnlineSubsystemGooglePlay.h"
 #include "OnlineError.h"
 #include <jni.h>
+#include "Algo/AllOf.h"
 #include "Android/AndroidJavaEnv.h"
 
 #define LOCTEXT_NAMESPACE "OnlineSubsystemGooglePlay"
 
-FGoogleTransactionData::FGoogleTransactionData(const FString& InOfferId, const FString& InPurchaseToken, const FString& InReceiptData, const FString& InSignature, EGooglePlayPurchaseState InPurchaseState)
-	: OfferId(InOfferId)
+FGoogleTransactionData::FGoogleTransactionData(const TArray<FString>& InOfferIds, const FString& InPurchaseToken, const FString& InReceiptData, const FString& InSignature, EGooglePlayPurchaseState InPurchaseState)
+	: OfferIds(InOfferIds)
 	, PurchaseToken(InPurchaseToken)
 	, CombinedTransactionData(InReceiptData, InSignature)
 	, PurchaseState(InPurchaseState)
@@ -30,14 +31,24 @@ FGoogleTransactionData::FGoogleTransactionData(const FString& InOfferId, const F
 
 FString FGoogleTransactionData::ToDebugString() const
 {
-	return FString::Printf(TEXT("OfferId: %s PurchaseState: %s TransactionId: %s ReceiptData: %s%s"),
-		*OfferId,
+	return FString::Printf(TEXT("OfferIds: %s PurchaseState: %s TransactionId: %s ReceiptData: %s%s"),
+		*FString::Join(OfferIds, TEXT(",")),
 		LexToString(PurchaseState),
 		*PurchaseToken,
 		*CombinedTransactionData.ToJson(),
 		ErrorStr.IsEmpty() ? TEXT("") : *FString::Printf(TEXT(" Error: %s"), *ErrorStr));
 }
 
+bool FGoogleTransactionData::IsMatchingRequest(const FPurchaseCheckoutRequest& Request) const
+{
+	return Algo::AllOf(GetOfferIds(), [&Request](const FString& OfferId) {
+				return Request.PurchaseOffers.ContainsByPredicate([&OfferId](const FPurchaseCheckoutRequest::FPurchaseOfferEntry& Entry) 
+				{
+					return OfferId == Entry.OfferId;
+				});
+		   }); 
+}
+
 FOnlinePurchaseGooglePlay::FOnlinePurchaseGooglePlay(FOnlineSubsystemGooglePlay* InSubsystem)
 	: bQueryingReceipts(false)
 	, Subsystem(InSubsystem)
@@ -69,6 +80,8 @@ public:
 	/** Add the single completed transaction to this transaction */
 	bool AddCompletedOffer(EPurchaseTransactionState Result, const FGoogleTransactionData& Transaction);
 
+	static void FillReceiptContent(FPurchaseReceipt& Receipt, const FGoogleTransactionData& Transaction);
+
 public:
 	
 	/** Checkout info for the in progress order */
@@ -119,6 +132,25 @@ TSharedRef<FPurchaseReceipt> FOnlinePurchaseGooglePlay::FOnlinePurchaseInProgres
 	return Receipt;
 }
 
+void FOnlinePurchaseGooglePlay::FOnlinePurchaseInProgressTransaction::FillReceiptContent(FPurchaseReceipt& Receipt, const FGoogleTransactionData& Transaction)
+{
+	// Pending purchases should not be redeemable although they already have a transaction identifier
+	const bool ShouldAddValidationInfo = (Transaction.GetPurchaseState() == EGooglePlayPurchaseState::Purchased);
+	for(const FString& OfferId: Transaction.GetOfferIds())
+	{
+		FPurchaseReceipt::FReceiptOfferEntry ReceiptEntry(TEXT(""), OfferId, 1);
+		FPurchaseReceipt::FLineItemInfo& LineItem = ReceiptEntry.LineItems.Emplace_GetRef();
+		
+		LineItem.ItemName = OfferId;
+		LineItem.UniqueId = Transaction.GetTransactionIdentifier();
+		if(ShouldAddValidationInfo)
+		{
+			LineItem.ValidationInfo = Transaction.GetCombinedReceiptData();
+		}
+		Receipt.AddReceiptOffer(MoveTemp(ReceiptEntry));
+	}
+}
+
 TSharedRef<FPurchaseReceipt> FOnlinePurchaseGooglePlay::FOnlinePurchaseInProgressTransaction::GenerateReceipt(const FGoogleTransactionData& Transaction)
 {
 	TSharedRef<FPurchaseReceipt> Receipt = MakeShared<FPurchaseReceipt>();
@@ -128,18 +160,7 @@ TSharedRef<FPurchaseReceipt> FOnlinePurchaseGooglePlay::FOnlinePurchaseInProgres
 	
 	if (Receipt->TransactionState == EPurchaseTransactionState::Purchased)
 	{
-		FPurchaseReceipt::FReceiptOfferEntry ReceiptEntry(TEXT(""), Transaction.GetOfferId(), 1);
-
-		FPurchaseReceipt::FLineItemInfo& LineItem = ReceiptEntry.LineItems.Emplace_GetRef();
-
-		LineItem.ItemName = Transaction.GetOfferId();
-		LineItem.UniqueId = Transaction.GetTransactionIdentifier();
-		if(Transaction.GetPurchaseState() == EGooglePlayPurchaseState::Purchased)
-		{
-			// Pending purchases should not be redeemable
-			LineItem.ValidationInfo = Transaction.GetCombinedReceiptData();
-		}
-		Receipt->AddReceiptOffer(MoveTemp(ReceiptEntry));
+		FillReceiptContent(*Receipt, Transaction);
 	}
 	
 	return Receipt;
@@ -147,28 +168,14 @@ TSharedRef<FPurchaseReceipt> FOnlinePurchaseGooglePlay::FOnlinePurchaseInProgres
 
 bool FOnlinePurchaseGooglePlay::FOnlinePurchaseInProgressTransaction::AddCompletedOffer(EPurchaseTransactionState Result, const FGoogleTransactionData& Transaction)
 {
-	bool bIsFromThisRequest = CheckoutRequest.PurchaseOffers.ContainsByPredicate([&Transaction](const FPurchaseCheckoutRequest::FPurchaseOfferEntry& Entry)
-	{
-		return (Transaction.GetOfferId() == Entry.OfferId);
-	});
+	bool bIsFromThisRequest = Transaction.IsMatchingRequest(CheckoutRequest);
 	
 	if(!bIsFromThisRequest)
 	{
 		return false;
 	}
 
-	FPurchaseReceipt::FReceiptOfferEntry Receipt(TEXT(""), Transaction.GetOfferId(), 1);
-
-	FPurchaseReceipt::FLineItemInfo& LineItem = Receipt.LineItems.Emplace_GetRef();
-	
-	LineItem.ItemName = Transaction.GetOfferId();
-	LineItem.UniqueId = Transaction.GetTransactionIdentifier();
-	if(Transaction.GetPurchaseState() == EGooglePlayPurchaseState::Purchased)
-	{
-		// Pending purchases should not be redeemable
-		LineItem.ValidationInfo = Transaction.GetCombinedReceiptData();
-	}
-	PurchaseInfo.AddReceiptOffer(MoveTemp(Receipt));
+	FillReceiptContent(PurchaseInfo, Transaction);
 	return true;
 }
 
@@ -206,15 +213,15 @@ void FOnlinePurchaseGooglePlay::Checkout(const FUniqueNetId& UserId, const FPurc
 		int32 NumOffers = CheckoutRequest.PurchaseOffers.Num();
 		if (NumOffers > 0)
 		{
-			const FPurchaseCheckoutRequest::FPurchaseOfferEntry& Offer = CheckoutRequest.PurchaseOffers[0];
-
-			extern bool AndroidThunkCpp_Iap_BeginPurchase(const FString&, const FString&);
+			extern bool AndroidThunkCpp_Iap_BeginPurchase(const TArray<FStringView>&, const FString&);
 			FString AccountId;
 			if (CheckoutRequest.AccountId.IsSet())
 			{
 				AccountId = CheckoutRequest.AccountId.GetValue();
 			}
-			bStarted = AndroidThunkCpp_Iap_BeginPurchase(Offer.OfferId, AccountId);
+			TArray<FStringView> ProductIds;
+			Algo::Transform(CheckoutRequest.PurchaseOffers, ProductIds, [](const FPurchaseCheckoutRequest::FPurchaseOfferEntry& Offer) { return FStringView(Offer.OfferId);});
+			bStarted = AndroidThunkCpp_Iap_BeginPurchase(ProductIds, AccountId);
 			if(bStarted)
 			{
 				UE_LOG_ONLINE_PURCHASE(Display, TEXT("Created a transaction."));
@@ -225,10 +232,6 @@ void FOnlinePurchaseGooglePlay::Checkout(const FUniqueNetId& UserId, const FPurc
 				UE_LOG_ONLINE_PURCHASE(Display, TEXT("Failed to create a transaction."));
 				RequestedTransaction->PurchaseInfo.TransactionState = EPurchaseTransactionState::Failed;
 			}
-			if (NumOffers > 1)
-			{
-				UE_LOG_ONLINE_PURCHASE(Warning, TEXT("Only supported purchasing one offer at a time, %d were requested and ignored"), NumOffers - 1);
-			}
 		}
 		else
 		{
@@ -335,10 +338,10 @@ void FOnlinePurchaseGooglePlay::OnTransactionCompleteResponse(EGooglePlayBilling
 	
 	if(InProgressTransaction != nullptr)
 	{
-		// Only supported purchasing one offer at a time
-		const FPurchaseCheckoutRequest::FPurchaseOfferEntry& Offer = InProgressTransaction->CheckoutRequest.PurchaseOffers[0];
-		bool bWasInProgressTransactionCanceled = (Result == EPurchaseTransactionState::Canceled && InTransactionData.GetOfferId().IsEmpty());
-		bIsInProgressTransaction = bWasInProgressTransactionCanceled || (Offer.OfferId == InTransactionData.GetOfferId());
+		bool bWasInProgressTransactionCanceled = (Result == EPurchaseTransactionState::Canceled && InTransactionData.GetOfferIds().IsEmpty());
+		bool bTransactionMatchesInProgressRequest = InTransactionData.IsMatchingRequest(InProgressTransaction->CheckoutRequest); 
+
+		bIsInProgressTransaction = bWasInProgressTransactionCanceled || bTransactionMatchesInProgressRequest;
 	}
 
 	if(bIsInProgressTransaction)
@@ -507,16 +510,16 @@ JNI_METHOD void Java_com_epicgames_unreal_GooglePlayStoreHelper_NativeQueryExist
 		for (jsize Idx = 0; Idx < NumProducts; Idx++)
 		{
 			// Build the product information strings.
-			const auto OfferId = FJavaHelper::FStringFromLocalRef(jenv, (jstring)jenv->GetObjectArrayElement(ProductIDs, Idx));
+			const auto OfferIds = FJavaHelper::ObjectArrayToFStringTArray(jenv, (jobjectArray)jenv->GetObjectArrayElement(ProductIDs, Idx));
 			const auto PurchaseToken = FJavaHelper::FStringFromLocalRef(jenv, (jstring)jenv->GetObjectArrayElement(PurchaseTokens, Idx));
 			const auto ReceiptData = FJavaHelper::FStringFromLocalRef(jenv, (jstring)jenv->GetObjectArrayElement(ReceiptsData, Idx));
 			const auto SignatureData = FJavaHelper::FStringFromLocalRef(jenv, (jstring)jenv->GetObjectArrayElement(Signatures, Idx));
 			const auto PurchaseState = (EGooglePlayPurchaseState)PurchaseStates[Idx];
 
-			FGoogleTransactionData ExistingPurchase(OfferId, PurchaseToken, ReceiptData, SignatureData, PurchaseState);
+			FGoogleTransactionData ExistingPurchase(OfferIds, PurchaseToken, ReceiptData, SignatureData, PurchaseState);
 			ExistingPurchaseInfo.Add(ExistingPurchase);
 
-			UE_LOG_ONLINE_PURCHASE(Verbose, TEXT("\Existing Product Identifier: %s"), *ExistingPurchase.ToDebugString());
+			UE_LOG_ONLINE_PURCHASE(Verbose, TEXT("Existing Product Identifier: %s"), *ExistingPurchase.ToDebugString());
 		}
 		jenv->ReleaseIntArrayElements(PurchaseStatesArray, PurchaseStates, JNI_ABORT);
 	}
@@ -531,18 +534,17 @@ JNI_METHOD void Java_com_epicgames_unreal_GooglePlayStoreHelper_NativeQueryExist
 	}
 }
 
-JNI_METHOD void Java_com_epicgames_unreal_GooglePlayStoreHelper_NativePurchaseComplete(JNIEnv* jenv, jobject /*Thiz*/, jint JavaResponseCode, jstring JavaProductId, jint JavaPurchaseState, jstring JavaPurchaseToken, jstring JavaReceiptData, jstring JavaSignature)
+JNI_METHOD void Java_com_epicgames_unreal_GooglePlayStoreHelper_NativePurchaseComplete(JNIEnv* jenv, jobject /*Thiz*/, jint JavaResponseCode, jobjectArray JavaProductIds, jint JavaPurchaseState, jstring JavaPurchaseToken, jstring JavaReceiptData, jstring JavaSignature)
 {
-	FString ProductId, PurchaseToken, ReceiptData, Signature;
 	EGooglePlayBillingResponseCode EGPResponse = (EGooglePlayBillingResponseCode)JavaResponseCode;
 	EGooglePlayPurchaseState PurchaseState = (EGooglePlayPurchaseState)JavaPurchaseState;
 
-	ProductId = FJavaHelper::FStringFromParam(jenv, JavaProductId);
-	PurchaseToken = FJavaHelper::FStringFromParam(jenv, JavaPurchaseToken);
-	ReceiptData = FJavaHelper::FStringFromParam(jenv, JavaReceiptData);
-	Signature = FJavaHelper::FStringFromParam(jenv, JavaSignature);
+	TArray<FString> ProductIds = FJavaHelper::ObjectArrayToFStringTArray(jenv, JavaProductIds);
+	FString PurchaseToken = FJavaHelper::FStringFromParam(jenv, JavaPurchaseToken);
+	FString ReceiptData = FJavaHelper::FStringFromParam(jenv, JavaReceiptData);
+	FString Signature = FJavaHelper::FStringFromParam(jenv, JavaSignature);
 
-	FGoogleTransactionData TransactionData(ProductId, PurchaseToken, ReceiptData, Signature, PurchaseState);
+	FGoogleTransactionData TransactionData(ProductIds, PurchaseToken, ReceiptData, Signature, PurchaseState);
 
 	UE_LOG_ONLINE_PURCHASE(Verbose, TEXT("NativePurchaseComplete received response code: %s. Transaction: %s\n"), LexToString(EGPResponse), *TransactionData.ToDebugString());
 
diff --git a/Engine/Plugins/Online/Android/OnlineSubsystemGooglePlay/Source/Public/OnlinePurchaseGooglePlay.h b/Engine/Plugins/Online/Android/OnlineSubsystemGooglePlay/Source/Public/OnlinePurchaseGooglePlay.h
--- a/Engine/Plugins/Online/Android/OnlineSubsystemGooglePlay/Source/Public/OnlinePurchaseGooglePlay.h
+++ b/Engine/Plugins/Online/Android/OnlineSubsystemGooglePlay/Source/Public/OnlinePurchaseGooglePlay.h
@@ -14,12 +14,14 @@ enum class EGooglePlayPurchaseState : uint8;
  */
 struct FGoogleTransactionData
 {
-	FGoogleTransactionData(const FString& InOfferId, const FString& InPurchaseToken, const FString& InReceiptData, const FString& InSignature, EGooglePlayPurchaseState InPurchaseState);
+	FGoogleTransactionData(const TArray<FString>& InOfferIds, const FString& InPurchaseToken, const FString& InReceiptData, const FString& InSignature, EGooglePlayPurchaseState InPurchaseState);
 
 	/** @return a string that prints useful debug information about this transaction */
 	FString ToDebugString() const;
 	/** @return offer id for this transaction */
-	const FString& GetOfferId() const { return OfferId; }
+	const FString& GetOfferId() const { return OfferIds[0]; }
+	/** @return all offer ids for this transaction */
+	const TArray<FString>& GetOfferIds() const { return OfferIds; }
 	/** @return receipt data for this transaction */
 	FString GetCombinedReceiptData() const { return CombinedTransactionData.ToJson(); }
 	/** @return receipt data for this transaction */
@@ -32,7 +34,8 @@ struct FGoogleTransactionData
 	const FString& GetTransactionIdentifier() const	{ return PurchaseToken; }
 	/** @return the purchase state */
 	EGooglePlayPurchaseState GetPurchaseState() const { return PurchaseState; }
-
+	/** Checks if all items reported in the transaction are present in the request */
+	bool IsMatchingRequest(const FPurchaseCheckoutRequest& Request) const;
 private:
 
 	/** Easy access to transmission of data required for backend validation */
@@ -60,7 +63,7 @@ private:
 	};
 
 	/** GooglePlay offer id */
-	FString OfferId;
+	TArray<FString> OfferIds;
 	/** PurchaseToken for the transaction */
 	FString PurchaseToken;
 	/** Error on the transaction, if applicable */
diff --git a/Engine/Plugins/Online/OnlineSubsystem/Source/Private/OnlineSubsystemImpl.cpp b/Engine/Plugins/Online/OnlineSubsystem/Source/Private/OnlineSubsystemImpl.cpp
--- a/Engine/Plugins/Online/OnlineSubsystem/Source/Private/OnlineSubsystemImpl.cpp
+++ b/Engine/Plugins/Online/OnlineSubsystem/Source/Private/OnlineSubsystemImpl.cpp
@@ -328,6 +328,9 @@ void FOnlineSubsystemImpl::FinalizeReceipts(const FUniqueNetId& UserId)
 			UE_LOG_ONLINE(Display, TEXT("Receipt: %s %s"),
 				*Receipt.TransactionId,
 				LexToString(Receipt.TransactionState));
+
+			TArray<TTuple<FStringView,FStringView>> AlreadyUsedFinalizeInfo;
+
 			for (const FPurchaseReceipt::FReceiptOfferEntry& ReceiptOffer : Receipt.ReceiptOffers)
 			{
 				UE_LOG_ONLINE(Display, TEXT(" -Namespace: %s Id: %s Quantity: %d"),
@@ -336,21 +339,30 @@ void FOnlineSubsystemImpl::FinalizeReceipts(const FUniqueNetId& UserId)
 					ReceiptOffer.Quantity);
 
 				UE_LOG_ONLINE(Display, TEXT(" -LineItems:"));
+				
 				for (const FPurchaseReceipt::FLineItemInfo& LineItem : ReceiptOffer.LineItems)
 				{
 					UE_LOG_ONLINE(Display, TEXT("  -Name: %s Id: %s ValidationInfo: %d bytes"),
 						*LineItem.ItemName,
 						*LineItem.UniqueId,
 						LineItem.ValidationInfo.Len());
-					if (LineItem.IsRedeemable())
+
+					TTuple<FStringView,FStringView> CurrentFinalizeInfo(LineItem.UniqueId, LineItem.ValidationInfo);
+					bool bAlreadyFinalized = AlreadyUsedFinalizeInfo.Find(CurrentFinalizeInfo) != INDEX_NONE;
+					if (bAlreadyFinalized)
 					{
-						UE_LOG_ONLINE(Display, TEXT("Finalizing %s!"), *Receipt.TransactionId);
-						PurchaseInt->FinalizePurchase(UserId, LineItem.UniqueId, LineItem.ValidationInfo);
+						UE_LOG_ONLINE(Display, TEXT("Already finalized by sibling LineItem"));					
 					}
-					else
+					else if (!LineItem.IsRedeemable())
 					{
 						UE_LOG_ONLINE(Display, TEXT("Not redeemable"));
 					}
+					else
+					{
+						UE_LOG_ONLINE(Display, TEXT("Finalizing %s!"), *Receipt.TransactionId);
+						PurchaseInt->FinalizePurchase(UserId, LineItem.UniqueId, LineItem.ValidationInfo);
+						AlreadyUsedFinalizeInfo.AddUnique(MoveTemp(CurrentFinalizeInfo));
+					}
 				}
 			}	
 		}
@@ -414,11 +426,13 @@ bool FOnlineSubsystemImpl::HandlePurchaseExecCommands(UWorld* InWorld, const TCH
 		}
 		else if (FParse::Command(&Cmd, TEXT("BUY")))
 		{
-			FString ProductId = FParse::Token(Cmd, false);
+			FString CommaSeparatedProductIds = FParse::Token(Cmd, false);
 			FString UserIdStr = FParse::Token(Cmd, false);
-			if (ProductId.IsEmpty() || UserIdStr.IsEmpty())
+			TArray<FString> ProductIds;
+			CommaSeparatedProductIds.ParseIntoArray(ProductIds, TEXT(","), true);
+			if (ProductIds.IsEmpty() || UserIdStr.IsEmpty())
 			{
-				UE_LOG_ONLINE(Warning, TEXT("usage: PURCHASE BUY <productid> <userid>"));
+				UE_LOG_ONLINE(Warning, TEXT("usage: PURCHASE BUY <comma_separated_productids> <userid>"));
 			}
 			else
 			{
@@ -427,7 +441,10 @@ bool FOnlineSubsystemImpl::HandlePurchaseExecCommands(UWorld* InWorld, const TCH
 				{
 					FPurchaseCheckoutRequest Request;
 					Request.AccountId = UserIdStr;
-					Request.AddPurchaseOffer(FString(), ProductId, 1);
+					for (const FString& ProductId : ProductIds)
+					{
+						Request.AddPurchaseOffer(FString(), ProductId, 1);
+					}
 					PurchaseInt->Checkout(*UserId, Request, FOnPurchaseCheckoutComplete::CreateLambda([](const FOnlineError& Result, const TSharedRef<FPurchaseReceipt>& Receipt)
 						{
 							UE_LOG_ONLINE(Log, TEXT("Checkout Result=%s TransactionId=%s TransactionState=%s"), *Result.ToLogString(), *Receipt->TransactionId, LexToString(Receipt->TransactionState));
diff --git a/Engine/Source/Runtime/Core/Private/Android/AndroidJavaEnv.cpp b/Engine/Source/Runtime/Core/Private/Android/AndroidJavaEnv.cpp
--- a/Engine/Source/Runtime/Core/Private/Android/AndroidJavaEnv.cpp
+++ b/Engine/Source/Runtime/Core/Private/Android/AndroidJavaEnv.cpp
@@ -277,6 +277,17 @@ FScopedJavaObject<jstring> FJavaHelper::ToJavaString(JNIEnv* Env, const FString&
 	return NewScopedJavaObject(Env, Env->NewStringUTF(TCHAR_TO_UTF8(*UnrealString)));
 }
 
+FScopedJavaObject<jobjectArray> FJavaHelper::ToJavaStringArray(JNIEnv* Env, const TArray<FStringView>& UnrealStrings)
+{
+	jclass JavaStringClass = AndroidJavaEnv::FindJavaClass("java/lang/String");
+	jobjectArray ObjectArray = Env->NewObjectArray((jsize)UnrealStrings.Num(), JavaStringClass, NULL);
+	for (int32 Idx = 0; Idx < UnrealStrings.Num(); ++Idx)
+	{
+		Env->SetObjectArrayElement(ObjectArray, Idx, Env->NewStringUTF(TCHAR_TO_UTF8(UnrealStrings[Idx].GetData())));
+	}
+	return NewScopedJavaObject(Env, ObjectArray);
+}
+
 TArray<FString> FJavaHelper::ObjectArrayToFStringTArray(JNIEnv* Env, jobjectArray ObjectArray)
 {
 	TArray<FString> ArrayOfStrings;
diff --git a/Engine/Source/Runtime/Core/Public/Android/AndroidJavaEnv.h b/Engine/Source/Runtime/Core/Public/Android/AndroidJavaEnv.h
--- a/Engine/Source/Runtime/Core/Public/Android/AndroidJavaEnv.h
+++ b/Engine/Source/Runtime/Core/Public/Android/AndroidJavaEnv.h
@@ -45,6 +45,18 @@ public:
 	{
 		if (*this)
 		{
+			if constexpr (std::is_same_v<T, jobjectArray>)
+			{
+				const jsize Length = Env->GetArrayLength(ObjRef);
+				for(jsize Idx = 0; Idx < Length; ++Idx)
+				{
+					jobject Element = Env->GetObjectArrayElement(ObjRef, Idx);
+					if (Element && !Env->IsSameObject(Element, NULL))
+					{
+						Env->DeleteLocalRef(Element);
+					}
+				}
+			}
 			Env->DeleteLocalRef(ObjRef);
 		}
 	}
@@ -93,6 +105,9 @@ public:
 	// Converts FString into a Java string wrapped in FScopedJavaObject
 	static FScopedJavaObject<jstring> ToJavaString(JNIEnv* Env, const FString& UnrealString);
 
+	// Converts a TArray<FStringView> into a Java string array wrapped in FScopedJavaObject. FStringView content is expected to be null terminated
+	static FScopedJavaObject<jobjectArray> ToJavaStringArray(JNIEnv* Env, const TArray<FStringView>& UnrealStrings);
+
 	// Converts the java objectArray to an array of FStrings. jopbjectArray must be a String[] on the Java side
 	static TArray<FString> ObjectArrayToFStringTArray(JNIEnv* Env, jobjectArray ObjectArray);
 };
diff --git a/Engine/Source/Runtime/Launch/Private/Android/AndroidJNI.cpp b/Engine/Source/Runtime/Launch/Private/Android/AndroidJNI.cpp
--- a/Engine/Source/Runtime/Launch/Private/Android/AndroidJNI.cpp
+++ b/Engine/Source/Runtime/Launch/Private/Android/AndroidJNI.cpp
@@ -213,7 +213,7 @@ void FJavaWrapper::FindGooglePlayBillingMethods(JNIEnv* Env)
 	JavaStringClass = FindClassGlobalRef(Env, "java/lang/String", false);
 	AndroidThunkJava_IapSetupService = FindMethod(Env, GoogleServicesClassID, "AndroidThunkJava_IapSetupService", "(Ljava/lang/String;)V", bIsStoreOptional);
 	AndroidThunkJava_IapQueryInAppPurchases = FindMethod(Env, GoogleServicesClassID, "AndroidThunkJava_IapQueryInAppPurchases", "([Ljava/lang/String;)Z", bIsStoreOptional);
-	AndroidThunkJava_IapBeginPurchase = FindMethod(Env, GoogleServicesClassID, "AndroidThunkJava_IapBeginPurchase", "(Ljava/lang/String;Ljava/lang/String;)Z", bIsStoreOptional);
+	AndroidThunkJava_IapBeginPurchase = FindMethod(Env, GoogleServicesClassID, "AndroidThunkJava_IapBeginPurchase", "([Ljava/lang/String;Ljava/lang/String;)Z", bIsStoreOptional);
 	AndroidThunkJava_IapIsAllowedToMakePurchases = FindMethod(Env, GoogleServicesClassID, "AndroidThunkJava_IapIsAllowedToMakePurchases", "()Z", bIsStoreOptional);
 	AndroidThunkJava_IapConsumePurchase = FindMethod(Env, GoogleServicesClassID, "AndroidThunkJava_IapConsumePurchase", "(Ljava/lang/String;)Z", bIsStoreOptional);
 	AndroidThunkJava_IapQueryExistingPurchases = FindMethod(Env, GoogleServicesClassID, "AndroidThunkJava_IapQueryExistingPurchases", "()Z", bIsStoreOptional);
@@ -1361,35 +1361,29 @@ bool AndroidThunkCpp_Iap_QueryInAppPurchases(const TArray<FString>& ProductIDs,
 	return AndroidThunkCpp_Iap_QueryInAppPurchases(ProductIDs);
 }
 
-bool AndroidThunkCpp_Iap_BeginPurchase(const FString& ProductID, const FString& AccountId)
+bool AndroidThunkCpp_Iap_BeginPurchase(const TArray<FStringView>& ProductIds, const FString& AccountId)
 {
-	FPlatformMisc::LowLevelOutputDebugStringf(TEXT("[JNI] - AndroidThunkCpp_Iap_BeginPurchase %s"), *ProductID);
+	FPlatformMisc::LowLevelOutputDebugStringf(TEXT("[JNI] - AndroidThunkCpp_Iap_BeginPurchase"));
 	bool bResult = false;
 	if (JNIEnv* Env = FAndroidApplication::GetJavaEnv())
 	{
 		CHECK_JNI_METHOD(FJavaWrapper::AndroidThunkJava_IapBeginPurchase);
 
-		auto ProductIdJava = FJavaHelper::ToJavaString(Env, ProductID);
+		auto ProductIdsJava = FJavaHelper::ToJavaStringArray(Env, ProductIds);
 		if (AccountId.IsEmpty())
 		{
-			bResult = FJavaWrapper::CallBooleanMethod(Env, FJavaWrapper::GoogleServicesThis, FJavaWrapper::AndroidThunkJava_IapBeginPurchase, *ProductIdJava, nullptr);
+			bResult = FJavaWrapper::CallBooleanMethod(Env, FJavaWrapper::GoogleServicesThis, FJavaWrapper::AndroidThunkJava_IapBeginPurchase, *ProductIdsJava, nullptr);
 		}
 		else
 		{
 			auto ObfuscatedAccountIdJava = FJavaHelper::ToJavaString(Env, AccountId);
-			bResult = FJavaWrapper::CallBooleanMethod(Env, FJavaWrapper::GoogleServicesThis, FJavaWrapper::AndroidThunkJava_IapBeginPurchase, *ProductIdJava, *ObfuscatedAccountIdJava);
+			bResult = FJavaWrapper::CallBooleanMethod(Env, FJavaWrapper::GoogleServicesThis, FJavaWrapper::AndroidThunkJava_IapBeginPurchase, *ProductIdsJava, *ObfuscatedAccountIdJava);
 		}
 	}
 
 	return bResult;
 }
 
-bool AndroidThunkCpp_Iap_BeginPurchase(const FString& ProductID, const bool bConsumable)
-{
-	FPlatformMisc::LowLevelOutputDebugString(TEXT("AndroidThunkCpp_Iap_BeginPurchase DEPRECATED, won't use consumable flag"));
-	return AndroidThunkCpp_Iap_BeginPurchase(ProductID, FString());
-}
-
 bool AndroidThunkCpp_Iap_ConsumePurchase(const FString& ProductToken)
 {
 	FPlatformMisc::LowLevelOutputDebugStringf(TEXT("[JNI] - AndroidThunkCpp_Iap_ConsumePurchase %s"), *ProductToken);
