Added support for subscriptions and full server to server validation to OnlinePurchaseGooglePlay.

Full server validation is enabled by setting ``bDisableLocalAcknowledgeAndConsume`` in ``OnlineSubsystemGooglePlay.Store`` section on the ini files. This way consume and acknowledge should be called after server validation through server to server calls to GooglePlay services and all product types allowed by GooglePlay store can be handled

Subscriptions are identified using a unique product id on GooglePlay like any other regular in app products. The main difference with a regular product is that the same subscription can be purchased by many different ways depending on the GooglePlay app configuration.
A subscription may have different base plans which are identified by a string id. Each base plan defines a base price, whetrer the plan is autorenewable or not, the length before expiration. For instance, a subscription may have a base plan with a price for a 1 month subscription and another plan with a different price for a 1 year subscription
Each base plan may have several offers which are identified by a string id. Those may be triggered on purchase by GooglePlay or by the developer. Those offers may configure several phases on a subscription depending on the number of times it has been renewed. For instance, an offer could set a Free period, then a discount on the first 3 renewals and then following renewals would be at the base plan price
Subscription product ids seen from the engine side have a special format starting with "s-" prefix (``-`` is not allowed in product ids in the GooglePlay store, so there would be no broken id on the licensees side) . In order to request product information (QueryOffersById) the product id can be requested with or without that prefix, but all the product ids reported to the engine will contain the prefix in case the product is a subscription
In order to identify base plans, offers and its price phases ``QueryOffersById`` will store a ``FOnlineStoreOffer``instance for each allowed combination
The product id format for those will be s-<store product id>:<base plan id>:<offer id>:<price phase index>
An instance for each base plan will be created without the offer id and price phase index. Also, an instance for each offer will be created without the price index part and containing the final price point(the last price point. This was explicitly requested) in addition to all the price steps
For example, this could be the list of ``FOnlineStoreOffer`` instances created by ``QueryOffersById`` requesting ``s-test_subscription_1``(which in the store would be identified as just ``test_subscription_1``). That would be a subscription with 2 plans, one of the with 2 offers. First offer would contain 3 phases and second just 2 phases
``s-test_subscription_1:base-plan1``
``s-test_subscription_1:base-plan1:offer-1``
``s-test_subscription_1:base-plan1:offer-1:0``
``s-test_subscription_1:base-plan1:offer-1:1``
``s-test_subscription_1:base-plan1:offer-1:2``
``s-test_subscription_1:base-plan1:offer-2``
``s-test_subscription_1:base-plan1:offer-2:0``
``s-test_subscription_1:base-plan1:offer-2:1``
``s-test_subscription_1:base-plan2``

To ``CheckOut`` a product any of those identifiers could be used. The price index would be discarded since a subscription offer must be started from the initial step (GooglePlay already keeps track if the user already used the Free trial period and removes it). If the identifier contains an offer id the subscription will be purchased using that offer

When receipts are queried using ``QueryReceipts`` we will receive receipts for any non consumed product, any pending product (user still needs to pay using an external method) and for any non expired subscriptions
GoogleBilling Library do not notify on device for expired/canceled/paused subscriptions. Real time notifications can be configured on servers to notify validation/authority servers about those changes and act accordingly, so no additional checks should be done from the client
Only notifications received are when a new product is purchased/becomes pending or a subscriptions is activated or resumed. In order to check for receipts status on device the only way should be to call ``QueryReceipts```again and note that there is no receipt for the given subscription. This should be done only in case validations are done locally on device (which is not secure, but there are many serverless games around in mobile)

diff --git a/Engine/Build/Android/Java/src/com/epicgames/unreal/GameActivity.java.template b/Engine/Build/Android/Java/src/com/epicgames/unreal/GameActivity.java.template
--- a/Engine/Build/Android/Java/src/com/epicgames/unreal/GameActivity.java.template
+++ b/Engine/Build/Android/Java/src/com/epicgames/unreal/GameActivity.java.template
@@ -4853,6 +4853,20 @@ public class GameActivity extends $${gameActivitySuperClass}$$ implements Surfac
 		return false;
 	}
 
+	public boolean AndroidThunkJava_IapAcknowledgePurchase(String purchaseToken)
+	{
+		Log.debug("[JAVA] - AndroidThunkJava_IapAcknowledgePurchase " + purchaseToken);
+
+		if( IapStoreHelper != null )
+		{
+			IapStoreHelper.AcknowledgePurchase(purchaseToken);
+			return true;
+		}
+
+		Log.debug("[JAVA] - Store Helper is invalid");
+		return false;
+	}
+
 	public boolean AndroidThunkJava_IapQueryExistingPurchases()
 	{
 		Log.debug("[JAVA] - AndroidThunkJava_IapQueryExistingPurchases");
@@ -6004,10 +6018,6 @@ public class GameActivity extends $${gameActivitySuperClass}$$ implements Surfac
 				symbols.setDecimalSeparator('.');
 				DecimalFormat decimalFormat = new DecimalFormat("###.##", symbols);
 
-				Display displayobj = GetDefaultDisplay();
-				String Test = ExtractFieldsToKVString(displayobj, ";", true);
-
-
 				String screenDPI = decimalFormat.format(metrics.xdpi) + "," + decimalFormat.format(metrics.ydpi) + "," + decimalFormat.format(metrics.densityDpi);
 				return screenDPI;
 			}
diff --git a/Engine/Build/Android/Java/src/com/epicgames/unreal/StoreHelper.java b/Engine/Build/Android/Java/src/com/epicgames/unreal/StoreHelper.java
--- a/Engine/Build/Android/Java/src/com/epicgames/unreal/StoreHelper.java
+++ b/Engine/Build/Android/Java/src/com/epicgames/unreal/StoreHelper.java
@@ -10,6 +10,7 @@ public interface StoreHelper
 	public boolean QueryInAppPurchases(String[] ProductIDs);
 	public boolean BeginPurchase(String[] ProductIDs, String ObfuscatedAccountId);
 	public boolean IsAllowedToMakePurchases();
+	public void AcknowledgePurchase(String purchaseToken);
 	public void ConsumePurchase(String purchaseToken);
 	public boolean QueryExistingPurchases();
 	public void	onDestroy();
diff --git a/Engine/Plugins/Online/Android/OnlineSubsystemGooglePlay/Source/Java/BillingApiV2/com/epicgames/unreal/GooglePlayStoreHelper.java b/Engine/Plugins/Online/Android/OnlineSubsystemGooglePlay/Source/Java/BillingApiV2/com/epicgames/unreal/GooglePlayStoreHelper.java
--- a/Engine/Plugins/Online/Android/OnlineSubsystemGooglePlay/Source/Java/BillingApiV2/com/epicgames/unreal/GooglePlayStoreHelper.java
+++ b/Engine/Plugins/Online/Android/OnlineSubsystemGooglePlay/Source/Java/BillingApiV2/com/epicgames/unreal/GooglePlayStoreHelper.java
@@ -2,6 +2,8 @@
 
 package com.epicgames.unreal;
 
+import android.util.Pair;
+
 import androidx.annotation.NonNull;
 import androidx.annotation.Nullable;
 
@@ -9,6 +11,8 @@ import com.android.billingclient.api.BillingClient;
 import com.android.billingclient.api.BillingClientStateListener;
 import com.android.billingclient.api.BillingFlowParams;
 import com.android.billingclient.api.BillingResult;
+import com.android.billingclient.api.AcknowledgePurchaseParams;
+import com.android.billingclient.api.AcknowledgePurchaseResponseListener;
 import com.android.billingclient.api.ConsumeParams;
 import com.android.billingclient.api.ConsumeResponseListener;
 import com.android.billingclient.api.ProductDetails;
@@ -22,7 +26,10 @@ import com.android.vending.billing.util.Base64; //WMM should use different base6
 
 import java.util.ArrayList;
 import java.util.Arrays;
+import java.util.HashSet;
 import java.util.List;
+import java.util.Objects;
+import java.util.Set;
 import java.util.stream.Collectors;
 
 public class GooglePlayStoreHelper implements StoreHelper, PurchasesUpdatedListener
@@ -39,13 +46,10 @@ public class GooglePlayStoreHelper implements StoreHelper, PurchasesUpdatedListe
 	// Cache access to the games activity.
 	private final GameActivity MainActivity;
 
-	// Custom error value. Should match with EGooglePlayBillingResponseCode::CustomLogicError
-	private final int CustomLogicErrorResponse = -127;
-
 	public GooglePlayStoreHelper(GameActivity InGameActivity, Logger InLog)
 	{
 		Log = InLog;
-		Log.debug("[GooglePlayStoreHelper] - GooglePlayStoreHelper::GooglePlayStoreHelper");
+		Log.debug("[GooglePlayStoreHelper] - GooglePlayStoreHelper");
 
 		MainActivity = InGameActivity;
 
@@ -56,8 +60,9 @@ public class GooglePlayStoreHelper implements StoreHelper, PurchasesUpdatedListe
 
 		Client.startConnection(new BillingClientStateListener() {
 			@Override
-			public void onBillingSetupFinished(@NonNull BillingResult billingResult) {
-				if (billingResult.getResponseCode() == BillingClient.BillingResponseCode.OK) 
+			public void onBillingSetupFinished(@NonNull BillingResult billingResult)
+			{
+				if (billingResult.getResponseCode() == BillingClient.BillingResponseCode.OK)
 				{
 					bIsIapSetup = true;
 					Log.debug("[GooglePlayStoreHelper] - In-app billing supported for " + MainActivity.getPackageName());
@@ -86,52 +91,91 @@ public class GooglePlayStoreHelper implements StoreHelper, PurchasesUpdatedListe
 	@Override
 	public boolean IsAllowedToMakePurchases()
 	{
-		Log.debug("[GooglePlayStoreHelper] - GooglePlayStoreHelper::IsAllowedToMakePurchases");
+		Log.debug("[GooglePlayStoreHelper] - IsAllowedToMakePurchases. Billing client connected: " + bIsIapSetup);
 		return bIsIapSetup;
 	}
 
+	/**
+	 * Gets a set of store product ids for all the subscriptions provided in knownProducts.
+	 */
+	private Set<String> GetSubscriptionStoreProductIds(List<ProductDetails> KnownProducts)
+	{
+		Set<String> SubscriptionIds = new HashSet<String>();
+		for(ProductDetails Details: KnownProducts)
+		{
+			if(Details.getProductType().equals(BillingClient.ProductType.SUBS))
+			{
+				SubscriptionIds.add(Details.getProductId());
+			}
+		}
+		return SubscriptionIds;
+	}
+
 	/**
 	 * Query product details for the provided product ids
 	 */
 	@Override
 	public boolean QueryInAppPurchases(String[] InProductIDs)
 	{
-		Log.debug("[GooglePlayStoreHelper] - GooglePlayStoreHelper::QueryInAppPurchases");
+		Log.debug("[GooglePlayStoreHelper] - QueryInAppPurchases");
 		Log.debug("[GooglePlayStoreHelper] - NumProducts: " + InProductIDs.length);
 		Log.debug("[GooglePlayStoreHelper] - Products: " + Arrays.toString(InProductIDs));
 
 		if (InProductIDs.length > 0)
 		{
-			QueryProductDetailsParams Params = CreateQueryProductDetailsParamsForProducts(InProductIDs);
-			
-			Client.queryProductDetailsAsync(Params, new ProductDetailsResponseListener() {
+			QueryProductDetailsParams QueryProducts = CreateQueryProductDetailsParams(InProductIDs, BillingClient.ProductType.INAPP);
+			QueryProductDetailsParams QuerySubs = CreateQueryProductDetailsParams(InProductIDs, BillingClient.ProductType.SUBS);
+
+			ProductDetailsResponseListener AggregatedListener = new ProductDetailsResponseListener() {
+				// This method is marked synchronized because although Google Billing Library documentation says "All methods annotated with AnyThread can be 
+				// called from any thread and all the asynchronous callbacks will be returned on the same thread" those callbacks may be called on different 
+				// threads if invoked from JNI.
 				@Override
-				public void onProductDetailsResponse(@NonNull BillingResult Result, @NonNull List<ProductDetails> ProductDetailsList)
+				public synchronized void onProductDetailsResponse(@NonNull BillingResult Result, @NonNull List<ProductDetails> ProductDetailsList)
 				{
+					if(AwaitingResponseCount <= 0)
+					{
+						return;
+					}
 					int ResponseCode = Result.getResponseCode();
-
 					if(ResponseCode == BillingClient.BillingResponseCode.OK)
 					{
-						Log.debug("[GooglePlayStoreHelper] - GooglePlayStoreHelper::QueryInAppPurchases - ProductDetails:" + ProductDetailsList.toString());
-						Log.debug("[GooglePlayStoreHelper] - GooglePlayStoreHelper::QueryInAppPurchases " + ProductDetailsList.size() + " items - Success!");
-
-						ListOfProductDetailsAsArrays DetailArrays = new ListOfProductDetailsAsArrays(ProductDetailsList);
-						NativeQueryComplete(BillingClient.BillingResponseCode.OK, DetailArrays.ProductIds, DetailArrays.Titles, DetailArrays.Descriptions, DetailArrays.Prices, DetailArrays.PricesRaw, DetailArrays.CurrencyCodes);
+						AllProductDetails.addAll(ProductDetailsList);
+						AwaitingResponseCount -= 1;
+						Log.debug("[GooglePlayStoreHelper] - QueryInAppPurchases - Received new " + ProductDetailsList.size() + " products. Waiting for additional responses: " + (AwaitingResponseCount > 0));
+						if (AwaitingResponseCount == 0)
+						{
+							Log.debug("[GooglePlayStoreHelper] - QueryInAppPurchases - ProductDetails:" + AllProductDetails.toString());
+							Log.debug("[GooglePlayStoreHelper] - QueryInAppPurchases " + AllProductDetails.size() + " items - Success!");
+
+							ListOfProductDetailsAsArrays DetailArrays = new ListOfProductDetailsAsArrays(AllProductDetails);
+							NativeQueryComplete(BillingClient.BillingResponseCode.OK, DetailArrays.ProductIds, DetailArrays.Titles, DetailArrays.Descriptions, DetailArrays.Prices, DetailArrays.PricesRaw, DetailArrays.CurrencyCodes);
+						}
 					}
 					else
 					{
-						Log.debug("[GooglePlayStoreHelper] - GooglePlayStoreHelper::QueryInAppPurchases - Failed with: " + TranslateBillingResult(Result));
+						AwaitingResponseCount = 0;
+						Log.debug("[GooglePlayStoreHelper] - QueryInAppPurchases - Failed with: " + TranslateBillingResult(Result));
 						NativeQueryComplete(ResponseCode, null, null, null, null, null, null);
 					}
 
-					Log.debug("[GooglePlayStoreHelper] - NativeQueryComplete done!");
+					if (AwaitingResponseCount == 0)
+					{
+						Log.debug("[GooglePlayStoreHelper] - NativeQueryComplete done!");
+					}
 				}
-			});
+
+				// We will be querying for in app products and subscriptions 
+				int AwaitingResponseCount = 2;
+				List<ProductDetails> AllProductDetails = new ArrayList<>();
+			};
+
+			Client.queryProductDetailsAsync(QueryProducts, AggregatedListener);
+			Client.queryProductDetailsAsync(QuerySubs, AggregatedListener);
 		}
 		else
 		{
 			Log.debug("[GooglePlayStoreHelper] - no products given to query");
-			NativeQueryComplete(CustomLogicErrorResponse, null, null, null, null, null, null);
 			return false;
 		}
 
@@ -144,27 +188,54 @@ public class GooglePlayStoreHelper implements StoreHelper, PurchasesUpdatedListe
 	@Override
 	public boolean BeginPurchase(final String[] ProductIds, final String ObfuscatedAccountId)
 	{
-		QueryProductDetailsParams Params = CreateQueryProductDetailsParamsForProducts(ProductIds);
+		if (ProductIds.length == 0)
+		{
+			Log.debug("[GooglePlayStoreHelper] - BeginPurchase - Failed! No product ids provided");
+			return false;
+
+		}
+
+		String CommonProductType = EncodedProductId.GetCommonProductType(ProductIds);
+		if (CommonProductType == null)
+		{
+			Log.debug("[GooglePlayStoreHelper] - BeginPurchase - Failed! All products should have the same product type");
+			return false;
+		}
 
-		Log.debug("[GooglePlayStoreHelper] - GooglePlayStoreHelper::BeginPurchase - querying product information " + String.join(",", ProductIds));
+		if (CommonProductType == BillingClient.ProductType.SUBS && ProductIds.length > 1)
+		{
+			// This condition was added because the error produced by BillingClient.launchBillingFlow was not informative
+			Log.debug("[GooglePlayStoreHelper] - BeginPurchase - Failed! Purchasing multiple subscriptions at once is not supported");
+			return false;
+		}
+
+		Log.debug("[GooglePlayStoreHelper] - BeginPurchase - querying product information " + String.join(",", ProductIds));
+		QueryProductDetailsParams Params = CreateQueryProductDetailsParams(ProductIds, CommonProductType);
 		Client.queryProductDetailsAsync(Params, new ProductDetailsResponseListener() {
 			@Override
 			public void onProductDetailsResponse(@NonNull BillingResult Result, @NonNull List<ProductDetails> ProductDetailsList)
 			{
 				if(Result.getResponseCode() == BillingClient.BillingResponseCode.OK)
 				{
+					Log.debug("[GooglePlayStoreHelper] - BeginPurchase - CreateBillingFlowParams with " + ProductDetailsList.stream().map(ProductDetails::getName).collect(Collectors.joining(",")));
+					final BillingFlowParams Params = CreateBillingFlowParams(ProductIds, ProductDetailsList, ObfuscatedAccountId);
+					if(Params == null)
+					{
+						Log.debug("[GooglePlayStoreHelper] - BeginPurchase - Failed! Matching products may not exist in the store");
+						NativePurchaseComplete(BillingClient.BillingResponseCode.DEVELOPER_ERROR, ProductIds, Purchase.PurchaseState.UNSPECIFIED_STATE,"", "", "");
+						return;
+					}
 					// BillingClient.launchBillingFlow is annotated @UiThread, so dispatch the call to the Activity UI thread
 					MainActivity.runOnUiThread(new Runnable()
 					{
 						@Override
 						public void run() {
-							Log.debug("[GooglePlayStoreHelper] - GooglePlayStoreHelper::BeginPurchase - Launching billing flow " + ProductDetailsList.stream().map(ProductDetails::getName).collect(Collectors.joining(",")));
-							BillingFlowParams Params = CreateBillingFlowParams(ProductDetailsList, ObfuscatedAccountId);
+							Log.debug("[GooglePlayStoreHelper] - BeginPurchase - launchBillingFlow " + ProductDetailsList.stream().map(ProductDetails::getName).collect(Collectors.joining(",")));
 							BillingResult Result = Client.launchBillingFlow(MainActivity, Params);
 							int ResponseCode = Result.getResponseCode();
-							if (ResponseCode != BillingClient.BillingResponseCode.OK) 
+							if (ResponseCode != BillingClient.BillingResponseCode.OK)
 							{
-								Log.debug("[GooglePlayStoreHelper] - GooglePlayStoreHelper::BeginPurchase - Failed! " + TranslateBillingResult(Result));
+								Log.debug("[GooglePlayStoreHelper] - BeginPurchase - Failed! " + TranslateBillingResult(Result));
 								NativePurchaseComplete(ResponseCode, ProductIds, Purchase.PurchaseState.UNSPECIFIED_STATE, "", "", "");
 							}
 						}
@@ -172,7 +243,7 @@ public class GooglePlayStoreHelper implements StoreHelper, PurchasesUpdatedListe
 				}
 				else
 				{
-					Log.debug("[GooglePlayStoreHelper] - GooglePlayStoreHelper::BeginPurchase - Failed! " + TranslateBillingResult(Result));
+					Log.debug("[GooglePlayStoreHelper] - BeginPurchase - Failed! " + TranslateBillingResult(Result));
 					NativePurchaseComplete(Result.getResponseCode(), ProductIds, Purchase.PurchaseState.UNSPECIFIED_STATE,"", "", "");
 				}
 			}
@@ -180,40 +251,115 @@ public class GooglePlayStoreHelper implements StoreHelper, PurchasesUpdatedListe
 		return true;
 	}
 
+
+	/**
+	 * Queries for all non consumed non expired purchases
+	 */
 	@Override
 	public boolean QueryExistingPurchases()
 	{
-		Log.debug("[GooglePlayStoreHelper] - GooglePlayStoreHelper::QueryExistingPurchases");
+		Log.debug("[GooglePlayStoreHelper] - QueryExistingPurchases");
 
-		QueryPurchasesParams Params =  QueryPurchasesParams.newBuilder()
+		QueryPurchasesParams QueryInApp =  QueryPurchasesParams.newBuilder()
 			.setProductType(BillingClient.ProductType.INAPP)
 			.build();
-		
-		Client.queryPurchasesAsync(Params, new PurchasesResponseListener() {
+		QueryPurchasesParams QuerySubs = QueryPurchasesParams.newBuilder()
+			.setProductType(BillingClient.ProductType.SUBS)
+			.build();
+
+		PurchasesResponseListener AggregatedListener = new PurchasesResponseListener() {
+			// This method is marked synchronized because although Google Billing Library documentation says "All methods annotated with AnyThread can be 
+			// called from any thread and all the asynchronous callbacks will be returned on the same thread" those callbacks may be called on different 
+			// threads if invoked from JNI 
 			@Override
-			public void onQueryPurchasesResponse(@NonNull BillingResult Result, @NonNull List<Purchase> OwnedItems) 
+			public synchronized void onQueryPurchasesResponse(@NonNull BillingResult Result, @NonNull List<Purchase> OwnedItems)
 			{
+				if (AwaitingResponseCount <= 0)
+					return;
+
 				int ResponseCode = Result.getResponseCode();
-				if (ResponseCode == BillingClient.BillingResponseCode.OK) 
+				if (ResponseCode == BillingClient.BillingResponseCode.OK)
 				{
-					Log.debug("[GooglePlayStoreHelper] - GooglePlayStoreHelper::QueryExistingPurchases - User has previously purchased " + OwnedItems.size() + " inapp products");
-					Log.debug("[GooglePlayStoreHelper] - GooglePlayStoreHelper::QueryExistingPurchases - Success!");
+					AllProducts.addAll(OwnedItems);
+					AwaitingResponseCount -= 1;
+					Log.debug("[GooglePlayStoreHelper] - QueryExistingPurchases - Received " + OwnedItems.size() + " products. Waiting for additional responses: " + (AwaitingResponseCount != 0));
+					if (AwaitingResponseCount == 0)
+					{
+						Log.debug("[GooglePlayStoreHelper] - QueryExistingPurchases - Success! User has previously purchased " + AllProducts.size() + " inapp products");
 
-					ListOfPurchasesAsArrays PurchasesArrays = new ListOfPurchasesAsArrays(OwnedItems);
-					NativeQueryExistingPurchasesComplete(ResponseCode, PurchasesArrays.OwnedProducts, PurchasesArrays.PurchaseStates, PurchasesArrays.ProductTokens, PurchasesArrays.Receipts, PurchasesArrays.Signatures);
-				} 
-				else 
+						if (AllProducts.isEmpty())
+						{
+							NativeQueryExistingPurchasesComplete(ResponseCode, new String[0][0], new int[0], new String[0], new String[0], new String[0]);
+						}
+						else
+						{
+							// Request product details so we can fix notified product ids in case they refer to a subscription
+							String[] StoreProductIds = AllProducts.stream().flatMap(Purchase -> Purchase.getProducts().stream()).toArray(String[]::new);
+							QueryProductDetailsParams QuerySubs = CreateQueryProductDetailsParams(StoreProductIds, BillingClient.ProductType.SUBS);
+							Client.queryProductDetailsAsync(QuerySubs, new ProductDetailsResponseListener() {
+								@Override
+								public void onProductDetailsResponse(@NonNull BillingResult Result, @NonNull List<ProductDetails> ProductDetailsList)
+								{
+									int ResponseCode = Result.getResponseCode();
+									if (ResponseCode == BillingClient.BillingResponseCode.OK)
+									{
+										Set<String> SubscriptionProductIds = GetSubscriptionStoreProductIds(ProductDetailsList);
+										ListOfPurchasesAsArrays PurchasesArrays = new ListOfPurchasesAsArrays(AllProducts, SubscriptionProductIds);
+										NativeQueryExistingPurchasesComplete(ResponseCode, PurchasesArrays.OwnedProducts, PurchasesArrays.PurchaseStates, PurchasesArrays.ProductTokens, PurchasesArrays.Receipts, PurchasesArrays.Signatures);
+									}
+									else
+									{
+										Log.debug("[GooglePlayStoreHelper] - QueryExistingPurchases - Failed to get product information while collecting existing purchases" + TranslateBillingResult(Result));
+										NativeQueryExistingPurchasesComplete(ResponseCode, null, null, null, null, null);
+									}
+								}
+							});
+						}
+					}
+				}
+				else
 				{
-					Log.debug("[GooglePlayStoreHelper] - GooglePlayStoreHelper::QueryExistingPurchases - Failed to collect existing purchases" + TranslateBillingResult(Result));
+					AwaitingResponseCount = 0;
+					Log.debug("[GooglePlayStoreHelper] - QueryExistingPurchases - Failed to collect existing purchases" + TranslateBillingResult(Result));
 					NativeQueryExistingPurchasesComplete(ResponseCode, null, null, null, null, null);
 				}
 			}
-		});
+			// We will be querying for in app products and subscriptions 
+			int AwaitingResponseCount = 2;
+			List<Purchase> AllProducts = new ArrayList<>();
+		};
+
+		Client.queryPurchasesAsync(QueryInApp, AggregatedListener);
+		Client.queryPurchasesAsync(QuerySubs, AggregatedListener);
 		return true;
 	}
 
+	/**
+	 * Locally acknowledges the purchase referred by the provided PurchaseToken
+	 */
 	@Override
-	public void ConsumePurchase(final String PurchaseToken)
+	public void AcknowledgePurchase(final String PurchaseToken)
+	{
+		Log.debug("[GooglePlayStoreHelper] - Beginning AcknowledgePurchase: " + PurchaseToken);
+
+		AcknowledgePurchaseParams Params = AcknowledgePurchaseParams.newBuilder()
+			.setPurchaseToken(PurchaseToken)
+			.build();
+		Client.acknowledgePurchase(Params, new AcknowledgePurchaseResponseListener() {
+			@Override
+			public void onAcknowledgePurchaseResponse(@NonNull BillingResult Result)
+			{
+				Log.debug("[GooglePlayStoreHelper] - AcknowledgePurchase response: " + TranslateBillingResult(Result));
+				NativeAcknowledgeComplete(Result.getResponseCode(), PurchaseToken);
+			}
+		});
+	}
+
+	/**
+	 * Locally consumes the purchase referred by the provided PurchaseToken
+	 */
+	@Override
+	public void ConsumePurchase(String PurchaseToken)
 	{
 		Log.debug("[GooglePlayStoreHelper] - Beginning ConsumePurchase: " + PurchaseToken);
 
@@ -222,7 +368,7 @@ public class GooglePlayStoreHelper implements StoreHelper, PurchasesUpdatedListe
 			.build();
 		Client.consumeAsync(Params, new ConsumeResponseListener() {
 			@Override
-			public void onConsumeResponse(@NonNull BillingResult Result, @NonNull String PurchaseToken) 
+			public void onConsumeResponse(@NonNull BillingResult Result, @NonNull String PurchaseToken)
 			{
 				Log.debug("[GooglePlayStoreHelper] - ConsumePurchase response: " + TranslateBillingResult(Result));
 				NativeConsumeComplete(Result.getResponseCode(), PurchaseToken);
@@ -234,32 +380,56 @@ public class GooglePlayStoreHelper implements StoreHelper, PurchasesUpdatedListe
 	// The Google's PurchasesUpdatedListener interface implementation.
 
 	@Override
-	public void onPurchasesUpdated(@NonNull BillingResult Result, @Nullable List<Purchase> UpdatedPurchases){
+	public void onPurchasesUpdated(@NonNull BillingResult Result, @Nullable final List<Purchase> UpdatedPurchases)
+	{
 		int responseCode = Result.getResponseCode();
-		Log.debug("[GooglePlayStoreHelper] - GooglePlayStoreHelper::onPurchasesUpdated - Processing purchase result. Response Code: " + TranslateBillingResult(Result));
-		if (responseCode == BillingClient.BillingResponseCode.OK && UpdatedPurchases != null) 
+		Log.debug("[GooglePlayStoreHelper] - onPurchasesUpdated - Processing purchase result. Response Code: " + TranslateBillingResult(Result));
+		if (responseCode == BillingClient.BillingResponseCode.OK && UpdatedPurchases != null)
 		{
-			for (final Purchase UpdatedPurchase : UpdatedPurchases) 
-			{
-				final String[] ProductIds = UpdatedPurchase.getProducts().toArray(new String[0]);
-				Log.debug("[GooglePlayStoreHelper] - GooglePlayStoreHelper::onPurchasesUpdated - Processing purchase of product " + String.join(",",ProductIds));
-				Log.debug("[GooglePlayStoreHelper] - Purchase data: " + UpdatedPurchase.toString());
-				Log.debug("[GooglePlayStoreHelper] - Data signature: " + UpdatedPurchase.getSignature());
+			// Request product details so we can fix notified product ids in case they refer to a subscription
+			String[] StoreProductIds = UpdatedPurchases.stream().flatMap(Purchase -> Purchase.getProducts().stream()).toArray(String[]::new);
+			QueryProductDetailsParams QuerySubs = CreateQueryProductDetailsParams(StoreProductIds, BillingClient.ProductType.SUBS);
+			Client.queryProductDetailsAsync(QuerySubs, new ProductDetailsResponseListener() {
+				@Override
+				public void onProductDetailsResponse(@NonNull BillingResult Result, @NonNull List<ProductDetails> ProductDetailsList)
+				{
+					if (responseCode == BillingClient.BillingResponseCode.OK)
+					{
+						Set<String> SubscriptionProductIds = GetSubscriptionStoreProductIds(ProductDetailsList);
 
-				String Receipt = Base64.encode(UpdatedPurchase.getOriginalJson().getBytes());
-				NativePurchaseComplete(BillingClient.BillingResponseCode.OK, ProductIds, UpdatedPurchase.getPurchaseState(), UpdatedPurchase.getPurchaseToken(), Receipt, UpdatedPurchase.getSignature());
-			}
+						for (Purchase UpdatedPurchase : UpdatedPurchases)
+						{
+							String[] ProductIds = EncodedProductId.StoreProductIdsToProductIds(UpdatedPurchase.getProducts(), SubscriptionProductIds);
+							Log.debug("[GooglePlayStoreHelper] - onPurchasesUpdated - Processing purchase of product " + String.join(",",ProductIds));
+							Log.debug("[GooglePlayStoreHelper] - Purchase data: " + UpdatedPurchase.toString());
+							Log.debug("[GooglePlayStoreHelper] - Data signature: " + UpdatedPurchase.getSignature());
+
+							String Receipt = Base64.encode(UpdatedPurchase.getOriginalJson().getBytes());
+							NativePurchaseComplete(BillingClient.BillingResponseCode.OK, ProductIds, UpdatedPurchase.getPurchaseState(), UpdatedPurchase.getPurchaseToken(), Receipt, UpdatedPurchase.getSignature());
+						}
+
+					}
+					else
+					{
+						Log.debug("[GooglePlayStoreHelper] - onPurchasesUpdated - Purchase failed while retrieving product information. Response code: " + TranslateBillingResult(Result));
+						NativePurchaseComplete(responseCode, null, Purchase.PurchaseState.UNSPECIFIED_STATE, "", "", "");
+					}
+				}
+			});
 		}
 		else
 		{
-			Log.debug("[GooglePlayStoreHelper] - GooglePlayStoreHelper::onPurchasesUpdated - Purchase failed. Response code: " + TranslateBillingResult(Result));
- 			NativePurchaseComplete(responseCode, null, Purchase.PurchaseState.UNSPECIFIED_STATE, "", "", "");
+			Log.debug("[GooglePlayStoreHelper] - onPurchasesUpdated - Purchase failed. Response code: " + TranslateBillingResult(Result));
+			NativePurchaseComplete(responseCode, null, Purchase.PurchaseState.UNSPECIFIED_STATE, "", "", "");
 		}
 	}
 
 	///////////////////////////////////////////////////////
 	// Some helpers
 
+	/**
+	 * Helper to convert Google Billing library received data to be passed to native calls
+	 */
 	class ListOfProductDetailsAsArrays
 	{
 		public String[] ProductIds;
@@ -269,38 +439,121 @@ public class GooglePlayStoreHelper implements StoreHelper, PurchasesUpdatedListe
 		public long[] PricesRaw;
 		public String[] CurrencyCodes;
 
-		public ListOfProductDetailsAsArrays(List<ProductDetails> InProductDetailsList)
+		private int ComputeTotalEntryCount(List<ProductDetails> InProductDetailsList)
 		{
-			ProductIds = new String[InProductDetailsList.size()];
-			Titles = new String[InProductDetailsList.size()];
-			Descriptions = new String[InProductDetailsList.size()];
-			Prices = new String[InProductDetailsList.size()];
-			PricesRaw = new long[InProductDetailsList.size()];
-			CurrencyCodes = new String[InProductDetailsList.size()];
-
-			for (int index = 0; index < InProductDetailsList.size(); index++) 
+			int TotalCount = 0;
+			for (ProductDetails ThisProduct : InProductDetailsList)
 			{
-				ProductDetails ThisProduct = InProductDetailsList.get(index);
-				ProductDetails.OneTimePurchaseOfferDetails PriceDetails = ThisProduct.getOneTimePurchaseOfferDetails();
+				switch(ThisProduct.getProductType())
+				{
+					case BillingClient.ProductType.INAPP:
+						TotalCount ++;
+						break;
+					case BillingClient.ProductType.SUBS:
+						for (ProductDetails.SubscriptionOfferDetails SubsOfferDetails: ThisProduct.getSubscriptionOfferDetails())
+						{
+							TotalCount += SubsOfferDetails.getPricingPhases().getPricingPhaseList().size();
+							if (SubsOfferDetails.getOfferId() != null)
+							{
+								// En extra entry will be added for offers without including price index on it's name
+								TotalCount += 1;
+							}
+						}
+						break;
+				}
+			}
+			return TotalCount;
+		}
 
-				Log.debug("[GooglePlayStoreHelper] - GooglePlayStoreHelper::QueryInAppPurchases - Parsing details for: " + ThisProduct.getProductId());
+		/**
+		 * Transforms store data into data consumable by the engine
+		 * For in app products it will create an entry with product id <store product id>
+		 * For subscription it will create an entry for each base plan with a product id <subscription prefix>:<store product id>:<base plan id>
+		 * For each offer in a subscription plan it will create an entry <subscription prefix>:<store product id>:<base plan id>:<offer id> with the 
+		 * same price data as the base plan. Then for each step in the offer it will create entries with id 
+		 * <subscription prefix>:<store product id>:<base plan id>:<offer id>:<price phase index> with increasing values for "price phase index" 
+		 */
+		public ListOfProductDetailsAsArrays(List<ProductDetails> InProductDetailsList)
+		{
+			int TotalEntryCount = ComputeTotalEntryCount(InProductDetailsList);
+			ProductIds = new String[TotalEntryCount];
+			Titles = new String[TotalEntryCount];
+			Descriptions = new String[TotalEntryCount];
+			Prices = new String[TotalEntryCount];
+			PricesRaw = new long[TotalEntryCount];
+			CurrencyCodes = new String[TotalEntryCount];
+
+			int Index = 0;
+			for (ProductDetails ThisProduct : InProductDetailsList)
+			{
+				Log.debug("[GooglePlayStoreHelper] - QueryInAppPurchases - Parsing details for: " + ThisProduct.getProductId());
 				Log.debug("[GooglePlayStoreHelper] - title: " + ThisProduct.getTitle());
 				Log.debug("[GooglePlayStoreHelper] - description: " + ThisProduct.getDescription());
-				Log.debug("[GooglePlayStoreHelper] - price: " + PriceDetails.getFormattedPrice());
-				Log.debug("[GooglePlayStoreHelper] - price_amount_micros: " + PriceDetails.getPriceAmountMicros());
-				Log.debug("[GooglePlayStoreHelper] - price_currency_code: " + PriceDetails.getPriceCurrencyCode());
-
-				ProductIds[index] = ThisProduct.getProductId();
-				Titles[index] = ThisProduct.getTitle();
-				Descriptions[index] = ThisProduct.getDescription();
-				Prices[index] = PriceDetails.getFormattedPrice();
-				PricesRaw[index] = PriceDetails.getPriceAmountMicros();
-				CurrencyCodes[index] = PriceDetails.getPriceCurrencyCode();
+				Titles[Index] = ThisProduct.getTitle();
+				Descriptions[Index] = ThisProduct.getDescription();
+
+				switch(ThisProduct.getProductType())
+				{
+					case BillingClient.ProductType.INAPP:
+						ProductDetails.OneTimePurchaseOfferDetails PriceDetails = ThisProduct.getOneTimePurchaseOfferDetails();
+						Log.debug("[GooglePlayStoreHelper] - price: " + PriceDetails.getFormattedPrice());
+						Log.debug("[GooglePlayStoreHelper] - price_amount_micros: " + PriceDetails.getPriceAmountMicros());
+						Log.debug("[GooglePlayStoreHelper] - price_currency_code: " + PriceDetails.getPriceCurrencyCode());
+
+						ProductIds[Index] = ThisProduct.getProductId();
+						Prices[Index] = PriceDetails.getFormattedPrice();
+						PricesRaw[Index] = PriceDetails.getPriceAmountMicros();
+						CurrencyCodes[Index] = PriceDetails.getPriceCurrencyCode();
+
+						Index ++;
+						break;
+					case BillingClient.ProductType.SUBS:
+						for (ProductDetails.SubscriptionOfferDetails SubsOfferDetails: ThisProduct.getSubscriptionOfferDetails())
+						{
+							int PricingPhaseIndex = 0;
+							for( ProductDetails.PricingPhase PricingPhase: SubsOfferDetails.getPricingPhases().getPricingPhaseList())
+							{
+								Log.debug("[GooglePlayStoreHelper] - price: " + PricingPhase.getFormattedPrice());
+								Log.debug("[GooglePlayStoreHelper] - price_amount_micros: " + PricingPhase.getPriceAmountMicros());
+								Log.debug("[GooglePlayStoreHelper] - price_currency_code: " + PricingPhase.getPriceCurrencyCode());
+
+								boolean bIsOffer = SubsOfferDetails.getOfferId() != null;
+								if (bIsOffer)
+								{
+									ProductIds[Index] = EncodedProductId.EncodeSubscriptionOfferPricePoint(ThisProduct.getProductId(), SubsOfferDetails.getBasePlanId(), SubsOfferDetails.getOfferId(), PricingPhaseIndex);
+								}
+								else
+								{
+									ProductIds[Index] = EncodedProductId.EncodeSubscriptionBasePlan(ThisProduct.getProductId(), SubsOfferDetails.getBasePlanId());
+								}
+								Prices[Index] = PricingPhase.getFormattedPrice();
+								PricesRaw[Index] = PricingPhase.getPriceAmountMicros();
+								CurrencyCodes[Index] = PricingPhase.getPriceCurrencyCode();
+
+								Index++;
+								PricingPhaseIndex++;
+
+								if (bIsOffer && PricingPhaseIndex == SubsOfferDetails.getPricingPhases().getPricingPhaseList().size())
+								{
+									// Add extra product information using last price point but not including price index on it's name
+									ProductIds[Index] = EncodedProductId.EncodeSubscriptionOffer(ThisProduct.getProductId(), SubsOfferDetails.getBasePlanId(), SubsOfferDetails.getOfferId());
+									Prices[Index] = PricingPhase.getFormattedPrice();
+									PricesRaw[Index] = PricingPhase.getPriceAmountMicros();
+									CurrencyCodes[Index] = PricingPhase.getPriceCurrencyCode();
+									Index++;
+								}
+							}
+						}
+						break;
+				}
 			}
 		}
 	}
 
-	class ListOfPurchasesAsArrays 
+	/**
+	 * Helper to convert Google Billing library received data to be passed to native calls
+	 */
+	static class ListOfPurchasesAsArrays
 	{
 		public String[][] OwnedProducts;
 		public int[] PurchaseStates;
@@ -308,7 +561,7 @@ public class GooglePlayStoreHelper implements StoreHelper, PurchasesUpdatedListe
 		public String[] Signatures;
 		public String[] Receipts;
 
-		public ListOfPurchasesAsArrays(List<Purchase> InOwnedItems) 
+		public ListOfPurchasesAsArrays(List<Purchase> InOwnedItems, Set<String> SubscriptionProductId)
 		{
 			OwnedProducts = new String[InOwnedItems.size()][];
 			PurchaseStates = new int[InOwnedItems.size()];
@@ -316,10 +569,10 @@ public class GooglePlayStoreHelper implements StoreHelper, PurchasesUpdatedListe
 			Signatures = new String[InOwnedItems.size()];
 			Receipts = new String[InOwnedItems.size()];
 
-			for (int index = 0; index < InOwnedItems.size(); index ++) 
+			for (int index = 0; index < InOwnedItems.size(); index ++)
 			{
 				Purchase OwnedPurchase = InOwnedItems.get(index);
-				OwnedProducts[index] = OwnedPurchase.getProducts().toArray(new String[0]);
+				OwnedProducts[index] = OwnedPurchase.getProducts().stream().map(ProductId -> EncodedProductId.StoreProductIdToProductId(ProductId, SubscriptionProductId)).toArray(String[]::new);
 				PurchaseStates[index] = OwnedPurchase.getPurchaseState();
 				ProductTokens[index] = OwnedPurchase.getPurchaseToken();
 				Signatures[index] = OwnedPurchase.getSignature();
@@ -327,15 +580,186 @@ public class GooglePlayStoreHelper implements StoreHelper, PurchasesUpdatedListe
 			}
 		}
 	}
-	
-	private static BillingFlowParams CreateBillingFlowParams(List<ProductDetails> InProductDetails, String InObfuscatedAccountId)
+
+
+	/*
+		Helper class to Encode/Decode subscription identity data on ProductIds
+		To properly select a subscription for purchase or get it's price information ProductIds passed to/received from the engine 
+		may have a prefix and encoded identifiers to include all that information as a single ProductId
+		On the GooglePlay Store each product is identified by a single id but subscriptions may define several base plans, which define renewal periods and base pricing.
+		Each base plan may define multiple offers which include special pricings (like free trials and special pricings after that)
+		Those subscription id are in the form
+		<subscription prefix><store product id>[:<base plan id>[:<offer id>[:<offer price index>]]]
+		To query for product information just and identifier like <subscription prefix><store product id> is enough to retrieve all base plans and offers
+		In order to start a purchase the <base plan id> part is mandatory. To purchase a given offer in a base plan the <offer id> is also mandatory
+	 */
+	static class EncodedProductId
 	{
-		ArrayList<BillingFlowParams.ProductDetailsParams> ProductDetailsParamList = new ArrayList<>(InProductDetails.size());
-		for(ProductDetails Details : InProductDetails) 
+		private final static String SUBSCRIPTION_PREFIX = "s-";
+
+		static public String EncodeSubscriptionOfferPricePoint(String StoreProductId, String BasePlanId, String OfferId, int PricingPhaseIndex)
+		{
+			return String.format("%s%s:%s:%s:%d", SUBSCRIPTION_PREFIX, StoreProductId, BasePlanId, OfferId, PricingPhaseIndex);
+		}
+		static public String EncodeSubscriptionBasePlan(String StoreProductId, String BasePlanId)
+		{
+			return String.format("%s%s:%s", SUBSCRIPTION_PREFIX, StoreProductId, BasePlanId);
+		}
+
+		static public String EncodeSubscriptionOffer(String StoreProductId, String BasePlanId, String OfferId)
+		{
+			return String.format("%s%s:%s:%s", SUBSCRIPTION_PREFIX, StoreProductId, BasePlanId, OfferId);
+		}
+
+		static public boolean IsSubscription(String InProductId)
 		{
-			ProductDetailsParamList.add(BillingFlowParams.ProductDetailsParams.newBuilder()
-				.setProductDetails(Details)
-				.build());
+			return InProductId.startsWith(SUBSCRIPTION_PREFIX);
+		}
+
+		/**
+		 * returns the product type that matches all product ids. In product ids refer to mixed product types it returns null
+		 */
+		static public String GetCommonProductType(String[] InProducts)
+		{
+			if (InProducts.length == 0)
+			{
+				return null;
+			}
+
+			boolean bIsFirstProductSubscription = IsSubscription(InProducts[0]);
+
+			for (int i=1; i < InProducts.length; ++i)
+			{
+				if (bIsFirstProductSubscription != IsSubscription(InProducts[i]))
+				{
+					return null;
+				}
+			}
+			return bIsFirstProductSubscription? BillingClient.ProductType.SUBS : BillingClient.ProductType.INAPP;
+		}
+
+		/**
+		 * Fixes store product id to reflect whether it refers to a subscription or not
+		 */
+		public static String StoreProductIdToProductId(String StoreProductId, Set<String> SubscriptionStoreProductIds)
+		{
+			return SubscriptionStoreProductIds.contains(StoreProductId)? SUBSCRIPTION_PREFIX + StoreProductId : StoreProductId;
+		}
+
+		/**
+		 * Fixes a list of store product id to reflect whether they refer to a subscription or not
+		 */
+		public static String[] StoreProductIdsToProductIds(List<String> StoreProductIds, Set<String> SubscriptionStoreProductIds)
+		{
+			return StoreProductIds.stream().map(StoreProductId -> StoreProductIdToProductId(StoreProductId, SubscriptionStoreProductIds)).toArray(String[]::new );
+		}
+
+		/**
+		 * Decodes the product identity data from a ProductId 
+		 */
+		EncodedProductId(String CompositeProductId)
+		{
+			bIsSubscription = IsSubscription(CompositeProductId);
+			if (bIsSubscription)
+			{
+				String[] Parts = CompositeProductId.substring(SUBSCRIPTION_PREFIX.length()).split(":");
+				if (Parts.length > 4)
+				{
+					// Unknown format. Get it as a full product id
+					StoreProductId = CompositeProductId;
+				}
+				else
+				{
+					// The PricingPhaseIndex is not needed to start a purchase or query for products, so no need to store it
+					if (Parts.length >= 1)
+					{
+						StoreProductId = Parts[0];
+					}
+					if (Parts.length >= 2)
+					{
+						BasePlanId = Parts[1];
+					}
+					if (Parts.length >= 3)
+					{
+						OfferId = Parts[2];
+					}
+				}
+			}
+			else
+			{
+				StoreProductId = CompositeProductId;
+			}
+		}
+
+		/**
+		 * Creates BillingFlowParams.ProductDetailsParams by looking for store objects retrieved from the store that match the identity data contained on this instance 
+		 * To create a BillingFlowParams.ProductDetailsParams for a subscription we need to use the ProductDetails instance for its store product id and the offer token 
+		 * from the nested SubscriptionOfferDetails list
+		 * To create a BillingFlowParams.ProductDetailsParams for a inapp product we only need the ProductDetails matching its store product id 
+		 * Returns null if we cannot create the proper instance
+		 */
+		public BillingFlowParams.ProductDetailsParams CreateProductDetailsParams(List<ProductDetails> InProductDetails)
+		{
+			if (bIsSubscription)
+			{
+				// If a ProductDetails matches the StoreProductId and also contains an offer matching BasePlanId and OfferId then build ProductDetailsParams from it and the offer token
+				return InProductDetails.stream()
+					.filter(Details -> Details.getProductId().equals(StoreProductId) && Details.getProductType().equals(BillingClient.ProductType.SUBS))
+					.flatMap(Details -> Details.getSubscriptionOfferDetails().stream().filter(Offer -> HasMatchingOffer(Offer)).map(Offer -> new Pair<>(Details, Offer.getOfferToken())))
+					.findFirst()
+					.map(OfferData -> BillingFlowParams.ProductDetailsParams.newBuilder().setProductDetails(OfferData.first).setOfferToken(OfferData.second).build())
+					.orElse(null);
+			}
+			else
+			{
+				// If a ProductDetails matches the StoreProductId then build ProductDetailsParams from it
+				return InProductDetails.stream()
+					.filter(Details -> Details.getProductId().equals(StoreProductId) && Details.getProductType().equals(BillingClient.ProductType.INAPP))
+					.findFirst()
+					.map(Details -> BillingFlowParams.ProductDetailsParams.newBuilder().setProductDetails(Details).build())
+					.orElse(null);
+			}
+		}
+
+		/**
+		 * Check if the offer data matches the identity data contained on this instance
+		 */
+		boolean HasMatchingOffer(ProductDetails.SubscriptionOfferDetails OfferDetails)
+		{
+			if (!OfferDetails.getBasePlanId().equals(BasePlanId))
+			{
+				return false;
+			}
+
+			if (OfferId == null || OfferId.isEmpty())
+			{
+				return OfferDetails.getOfferId() == null;
+			}
+			else
+			{
+				return OfferId.equals(OfferDetails.getOfferId());
+			}
+		}
+
+		boolean bIsSubscription;
+		String StoreProductId;
+		String BasePlanId;
+		String OfferId;
+	}
+
+	/**
+	 * Creates a BillingFlowParams from product ids with encoded identity data
+	 */
+	private static BillingFlowParams CreateBillingFlowParams(String[] RequestedProductIds, List<ProductDetails> InProductDetails, String InObfuscatedAccountId)
+	{
+		List<BillingFlowParams.ProductDetailsParams> ProductDetailsParamList = Arrays.stream(RequestedProductIds)
+			.map(ProductId -> new EncodedProductId(ProductId).CreateProductDetailsParams(InProductDetails))
+			.filter(Objects::nonNull)
+			.collect(Collectors.toList());
+
+		if (ProductDetailsParamList.isEmpty())
+		{
+			return null;
 		}
 		BillingFlowParams.Builder Params = BillingFlowParams.newBuilder()
 			.setProductDetailsParamsList(ProductDetailsParamList);
@@ -346,14 +770,19 @@ public class GooglePlayStoreHelper implements StoreHelper, PurchasesUpdatedListe
 		return Params.build();
 	}
 
-	private static QueryProductDetailsParams CreateQueryProductDetailsParamsForProducts(String[] InProductIds)
+	/**
+	 * Creates a QueryProductDetailsParams from product ids with encoded identity data. Only StoreProductId is needed to query information for all base plans 
+	 * and offers in a subscription prodduct
+	 */
+	private QueryProductDetailsParams CreateQueryProductDetailsParams(String[] InProductIds, String ProductType)
 	{
 		ArrayList<QueryProductDetailsParams.Product> ProductList = new ArrayList<>();
 		for (String ProductId : InProductIds)
 		{
+			EncodedProductId Info = new EncodedProductId(ProductId);
 			QueryProductDetailsParams.Product Product = QueryProductDetailsParams.Product.newBuilder()
-				.setProductId(ProductId)
-				.setProductType(BillingClient.ProductType.INAPP)
+				.setProductId(Info.StoreProductId)
+				.setProductType(ProductType)
 				.build();
 			ProductList.add(Product);
 		}
@@ -366,7 +795,6 @@ public class GooglePlayStoreHelper implements StoreHelper, PurchasesUpdatedListe
 	/**
 	 * Get a text translation of the Response Codes returned by google play.
 	 */
-
 	private static String TranslateBillingResponseCode(int ResponseCode)
 	{
 		switch(ResponseCode)
@@ -400,6 +828,9 @@ public class GooglePlayStoreHelper implements StoreHelper, PurchasesUpdatedListe
 		}
 	}
 
+	/**
+	 * Get an informative text with as much information about the BillingResult as possible.
+	 */
 	private static String TranslateBillingResult(@NonNull BillingResult Result)
 	{
 		int ResponseCode = Result.getResponseCode();
@@ -421,7 +852,7 @@ public class GooglePlayStoreHelper implements StoreHelper, PurchasesUpdatedListe
 	 */
 	public void onDestroy()
 	{
-		Log.debug("[GooglePlayStoreHelper] - GooglePlayStoreHelper::onDestroy");
+		Log.debug("[GooglePlayStoreHelper] - onDestroy");
 
 		if (Client != null)
 		{
@@ -434,4 +865,5 @@ public class GooglePlayStoreHelper implements StoreHelper, PurchasesUpdatedListe
 	public native void NativePurchaseComplete(int ResponseCode, String[] ProductIds, int PurchaseState, String ProductToken, String ReceiptData, String Signature);
 	public native void NativeQueryExistingPurchasesComplete(int ResponseCode, String[][] ProductIds, int[] PurchaseState, String[] ProductTokens, String[] ReceiptsData, String[] Signatures);
 	public native void NativeConsumeComplete(int ResponseCode, String PurchaseToken);
+	public native void NativeAcknowledgeComplete(int ResponseCode, String PurchaseToken);
 }
diff --git a/Engine/Plugins/Online/Android/OnlineSubsystemGooglePlay/Source/OnlineSubsystemGooglePlay_UPL.xml b/Engine/Plugins/Online/Android/OnlineSubsystemGooglePlay/Source/OnlineSubsystemGooglePlay_UPL.xml
--- a/Engine/Plugins/Online/Android/OnlineSubsystemGooglePlay/Source/OnlineSubsystemGooglePlay_UPL.xml
+++ b/Engine/Plugins/Online/Android/OnlineSubsystemGooglePlay/Source/OnlineSubsystemGooglePlay_UPL.xml
@@ -15,7 +15,7 @@
 	<buildGradleAdditions>
 		<insert>
 			dependencies {
-			implementation 'com.android.billingclient:billing:5.0.0'
+			implementation 'com.android.billingclient:billing:5.1.0'
 			}
 		</insert>
 	</buildGradleAdditions>
diff --git a/Engine/Plugins/Online/Android/OnlineSubsystemGooglePlay/Source/Private/OnlinePurchaseGooglePlay.cpp b/Engine/Plugins/Online/Android/OnlineSubsystemGooglePlay/Source/Private/OnlinePurchaseGooglePlay.cpp
--- a/Engine/Plugins/Online/Android/OnlineSubsystemGooglePlay/Source/Private/OnlinePurchaseGooglePlay.cpp
+++ b/Engine/Plugins/Online/Android/OnlineSubsystemGooglePlay/Source/Private/OnlinePurchaseGooglePlay.cpp
@@ -3,6 +3,7 @@
 #include "OnlinePurchaseGooglePlay.h"
 #include "OnlineSubsystemGooglePlay.h"
 #include "OnlineError.h"
+#include "Misc/ConfigCacheIni.h"
 #include <jni.h>
 #include "Algo/AllOf.h"
 #include "Android/AndroidJavaEnv.h"
@@ -12,7 +13,7 @@
 FGoogleTransactionData::FGoogleTransactionData(const TArray<FString>& InOfferIds, const FString& InPurchaseToken, const FString& InReceiptData, const FString& InSignature, EGooglePlayPurchaseState InPurchaseState)
 	: OfferIds(InOfferIds)
 	, PurchaseToken(InPurchaseToken)
-	, CombinedTransactionData(InReceiptData, InSignature)
+	, CombinedTransactionData(InOfferIds, InReceiptData, InSignature)
 	, PurchaseState(InPurchaseState)
 {
 	if (PurchaseToken.IsEmpty())
@@ -44,16 +45,25 @@ bool FGoogleTransactionData::IsMatchingRequest(const FPurchaseCheckoutRequest& R
 	return Algo::AllOf(GetOfferIds(), [&Request](const FString& OfferId) {
 				return Request.PurchaseOffers.ContainsByPredicate([&OfferId](const FPurchaseCheckoutRequest::FPurchaseOfferEntry& Entry) 
 				{
-					return OfferId == Entry.OfferId;
+					FStringView CheckedOutProductId(Entry.OfferId);
+					int32 ColonPos = INDEX_NONE;
+					if (CheckedOutProductId.FindChar(TEXT(':'), ColonPos))
+					{
+						CheckedOutProductId = CheckedOutProductId.Left(ColonPos);
+					}
+					return (OfferId == CheckedOutProductId);
 				});
 		   }); 
 }
 
-FOnlinePurchaseGooglePlay::FOnlinePurchaseGooglePlay(FOnlineSubsystemGooglePlay* InSubsystem)
-	: bQueryingReceipts(false)
-	, Subsystem(InSubsystem)
-{
-	UE_LOG_ONLINE_PURCHASE(Verbose, TEXT( "FOnlinePurchaseGooglePlay::FOnlinePurchaseGooglePlay" ));
+FGoogleTransactionData::FJsonReceiptData::FJsonReceiptData(const TArray<FString>& InOfferIds, const FString& InReceiptData, const FString& InSignature)
+	: ReceiptData(InReceiptData)
+	, Signature(InSignature)
+{ 
+	if (!InOfferIds.IsEmpty() && FOnlinePurchaseGooglePlay::IsSubscriptionProductId(InOfferIds[0]))
+	{
+		IsSubscription = true;
+	}
 }
 
 /**
@@ -179,6 +189,12 @@ bool FOnlinePurchaseGooglePlay::FOnlinePurchaseInProgressTransaction::AddComplet
 	return true;
 }
 
+FOnlinePurchaseGooglePlay::FOnlinePurchaseGooglePlay(FOnlineSubsystemGooglePlay* InSubsystem)
+	: Subsystem(InSubsystem)
+{
+	UE_LOG_ONLINE_PURCHASE(Verbose, TEXT( "FOnlinePurchaseGooglePlay::FOnlinePurchaseGooglePlay" ));
+}
+
 FOnlinePurchaseGooglePlay::~FOnlinePurchaseGooglePlay()
 {
 }
@@ -186,6 +202,22 @@ FOnlinePurchaseGooglePlay::~FOnlinePurchaseGooglePlay()
 void FOnlinePurchaseGooglePlay::Init()
 {
 	UE_LOG_ONLINE_PURCHASE(Verbose, TEXT("FOnlinePurchaseGooglePlay::Init"));
+
+	FString GooglePlayLicenseKey;
+	if (!GConfig->GetString(TEXT("/Script/AndroidRuntimeSettings.AndroidRuntimeSettings"), TEXT("GooglePlayLicenseKey"), GooglePlayLicenseKey, GEngineIni) || GooglePlayLicenseKey.IsEmpty())
+	{
+		UE_LOG_ONLINE_STOREV2(Warning, TEXT("Missing GooglePlayLicenseKey key in /Script/AndroidRuntimeSettings.AndroidRuntimeSettings of DefaultEngine.ini"));
+	}
+
+	GConfig->GetBool(TEXT("OnlineSubsystemGooglePlay.Store"), TEXT("bDisableLocalAcknowledgeAndConsume"), bDisableLocalAcknowledgeAndConsume, GEngineIni);
+
+	extern void AndroidThunkCpp_Iap_SetupIapService(const FString&);
+	AndroidThunkCpp_Iap_SetupIapService(GooglePlayLicenseKey);
+}
+
+bool FOnlinePurchaseGooglePlay::IsSubscriptionProductId(const FString& ProductId)
+{
+	return ProductId.StartsWith(TEXT("s-"));
 }
 
 bool FOnlinePurchaseGooglePlay::IsAllowedToPurchase(const FUniqueNetId& UserId)
@@ -270,9 +302,46 @@ void FOnlinePurchaseGooglePlay::Checkout(const FUniqueNetId& UserId, const FPurc
 
 void FOnlinePurchaseGooglePlay::FinalizePurchase(const FUniqueNetId& UserId, const FString& ReceiptId)
 {
-	UE_LOG_ONLINE_PURCHASE(Verbose, TEXT("FOnlinePurchaseGooglePlay::FinalizePurchase %s %s"), *UserId.ToString(), *ReceiptId);
-	extern void AndroidThunkCpp_Iap_ConsumePurchase(const FString&);
-	AndroidThunkCpp_Iap_ConsumePurchase(ReceiptId);
+	if (!bDisableLocalAcknowledgeAndConsume)
+	{
+		// Look for known receipt with given transaction id. Call overload with 3rd argument if found
+		TSharedRef<FPurchaseReceipt>* MatchingReceipt = KnownTransactions.FindByPredicate([&ReceiptId](const TSharedRef<FPurchaseReceipt>& Receipt)
+			{
+				return ReceiptId == Receipt->TransactionId;
+			});
+		if (MatchingReceipt != nullptr)
+		{
+			FinalizePurchase(UserId, ReceiptId, (*MatchingReceipt)->ReceiptOffers[0].LineItems[0].ValidationInfo);
+		}
+		else
+		{
+			// Always acknowledge to avoid failing on subscriptions
+			extern void AndroidThunkCpp_Iap_AcknowledgePurchase(const FString&);
+			AndroidThunkCpp_Iap_AcknowledgePurchase(ReceiptId);
+
+			UE_LOG_ONLINE_PURCHASE(Warning, TEXT("FOnlinePurchaseGooglePlay::Receipt not found matching transaction id %s. Will consume product to keep backwards compatibility"), *ReceiptId);
+			// Keep previous behaviour: All products are considered consumable. Subscriptions cannot be consumed. An error will be logged when trying to consume them
+			extern void AndroidThunkCpp_Iap_ConsumePurchase(const FString&);
+			AndroidThunkCpp_Iap_ConsumePurchase(ReceiptId);
+		}
+	}
+}
+
+void FOnlinePurchaseGooglePlay::FinalizePurchase(const FUniqueNetId& UserId, const FString& ReceiptId, const FString& ReceiptInfo)
+{
+	if (!bDisableLocalAcknowledgeAndConsume)
+	{
+		if(ReceiptInfo.Contains(FGoogleTransactionData::SubscriptionReceiptMarker))
+		{
+			extern void AndroidThunkCpp_Iap_AcknowledgePurchase(const FString&);
+			AndroidThunkCpp_Iap_AcknowledgePurchase(ReceiptId);
+		}
+		else
+		{
+			extern void AndroidThunkCpp_Iap_ConsumePurchase(const FString&);
+			AndroidThunkCpp_Iap_ConsumePurchase(ReceiptId);
+		}
+	}
 }
 
 void FOnlinePurchaseGooglePlay::RedeemCode(const FUniqueNetId& UserId, const FRedeemCodeRequest& RedeemCodeRequest, const FOnPurchaseRedeemCodeComplete& Delegate)
@@ -463,6 +532,11 @@ void FOnlinePurchaseGooglePlay::OnQueryExistingPurchasesComplete(EGooglePlayBill
 	}
 }
 
+void FOnlinePurchaseGooglePlay::OnAcknowledgePurchaseComplete(EGooglePlayBillingResponseCode InResponseCode, const FString& TransactionIdentifier)
+{
+	UE_LOG_ONLINE_PURCHASE(Log, TEXT("FOnlinePurchaseGooglePlay::OnConsumePurchaseComplete Response: %s PurchaseToken: %s"), LexToString(InResponseCode), *TransactionIdentifier);
+}
+
 void FOnlinePurchaseGooglePlay::OnConsumePurchaseComplete(EGooglePlayBillingResponseCode InResponseCode, const FString& TransactionIdentifier)
 {
 	UE_LOG_ONLINE_PURCHASE(Log, TEXT("FOnlinePurchaseGooglePlay::OnConsumePurchaseComplete Response: %s PurchaseToken: %s"), LexToString(InResponseCode), *TransactionIdentifier);
@@ -578,5 +652,25 @@ JNI_METHOD void Java_com_epicgames_unreal_GooglePlayStoreHelper_NativeConsumeCom
 	}
 }
 
+JNI_METHOD void Java_com_epicgames_unreal_GooglePlayStoreHelper_NativeAcknowledgeComplete(JNIEnv* jenv, jobject /*Thiz*/, jint JavaResponseCode, jstring JavaPurchaseToken)
+{
+	FString PurchaseToken;
+	EGooglePlayBillingResponseCode EGPResponse = (EGooglePlayBillingResponseCode)JavaResponseCode;
+
+	bool bWasSuccessful = (EGPResponse == EGooglePlayBillingResponseCode::Ok);
+	UE_LOG_ONLINE_PURCHASE(Verbose, TEXT("NativeAcknowledgeComplete received response code: %s. PurchaseToken: %s\n"), LexToString(EGPResponse), *PurchaseToken);
+
+	PurchaseToken = FJavaHelper::FStringFromParam(jenv, JavaPurchaseToken);
+
+	if (auto OnlineSubGP = static_cast<FOnlineSubsystemGooglePlay* const>(IOnlineSubsystem::Get(GOOGLEPLAY_SUBSYSTEM)))
+	{
+		OnlineSubGP->ExecuteNextTick([OnlineSubGP, EGPResponse, PurchaseToken]()
+		{
+			TSharedPtr<FOnlinePurchaseGooglePlay> PurchaseInt = StaticCastSharedPtr<FOnlinePurchaseGooglePlay>(OnlineSubGP->GetPurchaseInterface());
+			PurchaseInt->OnAcknowledgePurchaseComplete(EGPResponse, PurchaseToken);
+		});
+	}
+}
+
 #undef LOCTEXT_NAMESPACE
 
diff --git a/Engine/Plugins/Online/Android/OnlineSubsystemGooglePlay/Source/Private/OnlineStoreGooglePlay.cpp b/Engine/Plugins/Online/Android/OnlineSubsystemGooglePlay/Source/Private/OnlineStoreGooglePlay.cpp
--- a/Engine/Plugins/Online/Android/OnlineSubsystemGooglePlay/Source/Private/OnlineStoreGooglePlay.cpp
+++ b/Engine/Plugins/Online/Android/OnlineSubsystemGooglePlay/Source/Private/OnlineStoreGooglePlay.cpp
@@ -6,7 +6,6 @@
 #include "Internationalization/Internationalization.h"
 #include "Internationalization/Culture.h"
 #include "Internationalization/FastDecimalFormat.h"
-#include "Misc/ConfigCacheIni.h"
 #include <jni.h>
 #include "Android/AndroidJavaEnv.h"
 
@@ -28,20 +27,6 @@ FOnlineStoreGooglePlayV2::~FOnlineStoreGooglePlayV2()
 {
 }
 
-void FOnlineStoreGooglePlayV2::Init()
-{
-	UE_LOG_ONLINE_STOREV2(Verbose, TEXT("FOnlineStoreGooglePlayV2::Init"));
-
-	FString GooglePlayLicenseKey;
-	if (!GConfig->GetString(TEXT("/Script/AndroidRuntimeSettings.AndroidRuntimeSettings"), TEXT("GooglePlayLicenseKey"), GooglePlayLicenseKey, GEngineIni) || GooglePlayLicenseKey.IsEmpty())
-	{
-		UE_LOG_ONLINE_STOREV2(Warning, TEXT("Missing GooglePlayLicenseKey key in /Script/AndroidRuntimeSettings.AndroidRuntimeSettings of DefaultEngine.ini"));
-	}
-
-	extern void AndroidThunkCpp_Iap_SetupIapService(const FString&);
-	AndroidThunkCpp_Iap_SetupIapService(GooglePlayLicenseKey);
-}
-
 TSharedRef<FOnlineStoreOffer> ConvertProductToStoreOffer(const FOnlineStoreOffer& Product)
 {
 	return MakeShared<FOnlineStoreOffer>(Product);
@@ -120,7 +105,10 @@ void FOnlineStoreGooglePlayV2::QueryOffersById(const FUniqueNetId& UserId, const
 		bIsQueryInFlight = true;
 		QueryOnlineStoreOffersCompleteDelegate = Delegate;
 		extern bool AndroidThunkCpp_Iap_QueryInAppPurchases(const TArray<FString>&);
-		AndroidThunkCpp_Iap_QueryInAppPurchases(OfferIds);
+		if (!AndroidThunkCpp_Iap_QueryInAppPurchases(OfferIds))
+		{
+			Delegate.ExecuteIfBound(false, OfferIds, TEXT("Query did not start"));
+		}
 	}
 }
 
diff --git a/Engine/Plugins/Online/Android/OnlineSubsystemGooglePlay/Source/Private/OnlineSubsystemGooglePlay.cpp b/Engine/Plugins/Online/Android/OnlineSubsystemGooglePlay/Source/Private/OnlineSubsystemGooglePlay.cpp
--- a/Engine/Plugins/Online/Android/OnlineSubsystemGooglePlay/Source/Private/OnlineSubsystemGooglePlay.cpp
+++ b/Engine/Plugins/Online/Android/OnlineSubsystemGooglePlay/Source/Private/OnlineSubsystemGooglePlay.cpp
@@ -141,7 +141,6 @@ bool FOnlineSubsystemGooglePlay::Init()
 	if (IsInAppPurchasingEnabled())
 	{
         StoreV2Interface = MakeShareable(new FOnlineStoreGooglePlayV2(this));
-		StoreV2Interface->Init();
 		PurchaseInterface = MakeShareable(new FOnlinePurchaseGooglePlay(this));
 		PurchaseInterface->Init();
 	}
diff --git a/Engine/Plugins/Online/Android/OnlineSubsystemGooglePlay/Source/Public/OnlinePurchaseGooglePlay.h b/Engine/Plugins/Online/Android/OnlineSubsystemGooglePlay/Source/Public/OnlinePurchaseGooglePlay.h
--- a/Engine/Plugins/Online/Android/OnlineSubsystemGooglePlay/Source/Public/OnlinePurchaseGooglePlay.h
+++ b/Engine/Plugins/Online/Android/OnlineSubsystemGooglePlay/Source/Public/OnlinePurchaseGooglePlay.h
@@ -3,6 +3,7 @@
 #pragma once
 
 #include "Interfaces/OnlinePurchaseInterface.h"
+#include "Misc/Optional.h"
 #include "Serialization/JsonSerializerMacros.h"
 #include "OnlineIdentityInterfaceGooglePlay.h"
 
@@ -36,6 +37,8 @@ struct FGoogleTransactionData
 	EGooglePlayPurchaseState GetPurchaseState() const { return PurchaseState; }
 	/** Checks if all items reported in the transaction are present in the request */
 	bool IsMatchingRequest(const FPurchaseCheckoutRequest& Request) const;
+	/** marker found in the receipt data to indicate it refers to a subscription */
+	inline static const TCHAR* SubscriptionReceiptMarker = TEXT("isSubscription");
 private:
 
 	/** Easy access to transmission of data required for backend validation */
@@ -45,11 +48,10 @@ private:
 	public:
 
 		FJsonReceiptData() {}
-		FJsonReceiptData(const FString& InReceiptData, const FString& InSignature)
-			: ReceiptData(InReceiptData)
-			, Signature(InSignature)
-		{ }
+		FJsonReceiptData(const TArray<FString>& InOfferIds, const FString& InReceiptData, const FString& InSignature);
 
+		/** Opaque store receipt data */
+		TOptional<bool> IsSubscription;
 		/** Opaque store receipt data */
 		FString ReceiptData;
 		/** Signature associated with the transaction */
@@ -57,6 +59,7 @@ private:
 
 		// FJsonSerializable
 		BEGIN_JSON_SERIALIZER
+			JSON_SERIALIZE_OPTIONAL("isSubscription", IsSubscription);
 			JSON_SERIALIZE("receiptData", ReceiptData);
 			JSON_SERIALIZE("signature", Signature);
 		END_JSON_SERIALIZER
@@ -90,6 +93,7 @@ public:
 	virtual void Checkout(const FUniqueNetId& UserId, const FPurchaseCheckoutRequest& CheckoutRequest, const FOnPurchaseCheckoutComplete& Delegate) override;
 	virtual void Checkout(const FUniqueNetId& UserId, const FPurchaseCheckoutRequest& CheckoutRequest, const FOnPurchaseReceiptlessCheckoutComplete& Delegate) override;
 	virtual void FinalizePurchase(const FUniqueNetId& UserId, const FString& ReceiptId) override;
+	virtual void FinalizePurchase(const FUniqueNetId& UserId, const FString& ReceiptId, const FString& ReceiptInfo) override;
 	virtual void RedeemCode(const FUniqueNetId& UserId, const FRedeemCodeRequest& RedeemCodeRequest, const FOnPurchaseRedeemCodeComplete& Delegate) override;
 	virtual void QueryReceipts(const FUniqueNetId& UserId, bool bRestoreReceipts, const FOnQueryReceiptsComplete& Delegate) override;
 	virtual void GetReceipts(const FUniqueNetId& UserId, TArray<FPurchaseReceipt>& OutReceipts) const override;
@@ -118,8 +122,13 @@ public:
 	/** Handle Java side query purchases completed notification */
 	void OnQueryExistingPurchasesComplete(EGooglePlayBillingResponseCode InResponseCode, const TArray<FGoogleTransactionData>& InExistingPurchases);
 
+	/** Handle Java side acknowledge purchase completed notification */
+	void OnAcknowledgePurchaseComplete(EGooglePlayBillingResponseCode InResponseCode, const FString& InPurchaseToken);
+
 	/** Handle Java side consume purchase completed notification */
 	void OnConsumePurchaseComplete(EGooglePlayBillingResponseCode InResponseCode, const FString& InPurchaseToken);
+
+	static bool IsSubscriptionProductId(const FString& ProductId);
 private:
 	
 	/**
@@ -132,8 +141,15 @@ private:
 	
 private:
 	
+	/** 
+	 * Acknoledge and consume are not invoked when calling FinishPurchase 
+	 * If this is set to 'true' acknowledge and consume should be invoked using server to server calls
+	 * after validation to avoid refunds
+	*/
+	bool bDisableLocalAcknowledgeAndConsume = false;
+
 	/** Are receipts being queried */
-	bool bQueryingReceipts;
+	bool bQueryingReceipts = false;
 	
 	/** Transient delegate to fire when query receipts has completed */
 	FOnQueryReceiptsComplete QueryReceiptsComplete;
@@ -145,7 +161,7 @@ private:
 	FOnlinePurchasePurchasedTransactions KnownTransactions;
 	
 	/** Reference to the parent subsystem */
-	FOnlineSubsystemGooglePlay* Subsystem;
+	FOnlineSubsystemGooglePlay* Subsystem = nullptr;
 };
 
 typedef TSharedPtr<FOnlinePurchaseGooglePlay, ESPMode::ThreadSafe> FOnlinePurchaseGooglePlayPtr;
diff --git a/Engine/Plugins/Online/Android/OnlineSubsystemGooglePlay/Source/Public/OnlineStoreGooglePlay.h b/Engine/Plugins/Online/Android/OnlineSubsystemGooglePlay/Source/Public/OnlineStoreGooglePlay.h
--- a/Engine/Plugins/Online/Android/OnlineSubsystemGooglePlay/Source/Public/OnlineStoreGooglePlay.h
+++ b/Engine/Plugins/Online/Android/OnlineSubsystemGooglePlay/Source/Public/OnlineStoreGooglePlay.h
@@ -45,9 +45,6 @@ public:
 	 */
 	virtual ~FOnlineStoreGooglePlayV2();
 
-	/** Initialize the interface */
-	void Init();
-	
 	/** Handle Java side query available iap completed notification */
 	void OnGooglePlayAvailableIAPQueryComplete(EGooglePlayBillingResponseCode InResponse, const TArray<FProvidedProductInformation>& InProvidedProductInformation);
 
diff --git a/Engine/Source/Runtime/Core/Private/Android/AndroidJavaEnv.cpp b/Engine/Source/Runtime/Core/Private/Android/AndroidJavaEnv.cpp
--- a/Engine/Source/Runtime/Core/Private/Android/AndroidJavaEnv.cpp
+++ b/Engine/Source/Runtime/Core/Private/Android/AndroidJavaEnv.cpp
@@ -283,7 +283,15 @@ FScopedJavaObject<jobjectArray> FJavaHelper::ToJavaStringArray(JNIEnv* Env, cons
 	jobjectArray ObjectArray = Env->NewObjectArray((jsize)UnrealStrings.Num(), JavaStringClass, NULL);
 	for (int32 Idx = 0; Idx < UnrealStrings.Num(); ++Idx)
 	{
-		Env->SetObjectArrayElement(ObjectArray, Idx, Env->NewStringUTF(TCHAR_TO_UTF8(UnrealStrings[Idx].GetData())));
+		// FStringView of an empty FString contains a null pointer as data
+		if (UnrealStrings[Idx].GetData())
+		{
+			Env->SetObjectArrayElement(ObjectArray, Idx, Env->NewStringUTF(TCHAR_TO_UTF8(UnrealStrings[Idx].GetData())));
+		}
+		else
+		{		
+			Env->SetObjectArrayElement(ObjectArray, Idx, Env->NewStringUTF(""));
+		}
 	}
 	return NewScopedJavaObject(Env, ObjectArray);
 }
diff --git a/Engine/Source/Runtime/Launch/Private/Android/AndroidJNI.cpp b/Engine/Source/Runtime/Launch/Private/Android/AndroidJNI.cpp
--- a/Engine/Source/Runtime/Launch/Private/Android/AndroidJNI.cpp
+++ b/Engine/Source/Runtime/Launch/Private/Android/AndroidJNI.cpp
@@ -221,6 +221,7 @@ void FJavaWrapper::FindGooglePlayBillingMethods(JNIEnv* Env)
 	AndroidThunkJava_IapQueryInAppPurchases = FindMethod(Env, GoogleServicesClassID, "AndroidThunkJava_IapQueryInAppPurchases", "([Ljava/lang/String;)Z", bIsStoreOptional);
 	AndroidThunkJava_IapBeginPurchase = FindMethod(Env, GoogleServicesClassID, "AndroidThunkJava_IapBeginPurchase", "([Ljava/lang/String;Ljava/lang/String;)Z", bIsStoreOptional);
 	AndroidThunkJava_IapIsAllowedToMakePurchases = FindMethod(Env, GoogleServicesClassID, "AndroidThunkJava_IapIsAllowedToMakePurchases", "()Z", bIsStoreOptional);
+	AndroidThunkJava_IapAcknowledgePurchase = FindMethod(Env, GoogleServicesClassID, "AndroidThunkJava_IapAcknowledgePurchase", "(Ljava/lang/String;)Z", bIsStoreOptional);
 	AndroidThunkJava_IapConsumePurchase = FindMethod(Env, GoogleServicesClassID, "AndroidThunkJava_IapConsumePurchase", "(Ljava/lang/String;)Z", bIsStoreOptional);
 	AndroidThunkJava_IapQueryExistingPurchases = FindMethod(Env, GoogleServicesClassID, "AndroidThunkJava_IapQueryExistingPurchases", "()Z", bIsStoreOptional);
 }
@@ -454,6 +455,7 @@ jmethodID FJavaWrapper::AndroidThunkJava_IapQueryInAppPurchases;
 jmethodID FJavaWrapper::AndroidThunkJava_IapBeginPurchase;
 jmethodID FJavaWrapper::AndroidThunkJava_IapIsAllowedToMakePurchases;
 jmethodID FJavaWrapper::AndroidThunkJava_IapQueryExistingPurchases;
+jmethodID FJavaWrapper::AndroidThunkJava_IapAcknowledgePurchase;
 jmethodID FJavaWrapper::AndroidThunkJava_IapConsumePurchase;
 
 jmethodID FJavaWrapper::AndroidThunkJava_UseSurfaceViewWorkaround;
@@ -1426,6 +1428,27 @@ bool AndroidThunkCpp_Iap_ConsumePurchase(const FString& ProductToken)
 	return bResult;
 }
 
+bool AndroidThunkCpp_Iap_AcknowledgePurchase(const FString& ProductToken)
+{
+	FPlatformMisc::LowLevelOutputDebugStringf(TEXT("[JNI] - AndroidThunkCpp_Iap_AcknowledgePurchase %s"), *ProductToken);
+	
+	bool bResult = false;
+	if (!ProductToken.IsEmpty())
+	{
+		if (JNIEnv* Env = FAndroidApplication::GetJavaEnv())
+		{
+			CHECK_JNI_METHOD(FJavaWrapper::AndroidThunkJava_IapAcknowledgePurchase);
+
+			auto ProductTokenJava = FJavaHelper::ToJavaString(Env, ProductToken);
+			//FPlatformMisc::LowLevelOutputDebugStringf(TEXT("[JNI] - AndroidThunkCpp_Iap_AcknowledgePurchase BEGIN"));
+			bResult = FJavaWrapper::CallBooleanMethod(Env, FJavaWrapper::GoogleServicesThis, FJavaWrapper::AndroidThunkJava_IapAcknowledgePurchase, *ProductTokenJava);
+			//FPlatformMisc::LowLevelOutputDebugStringf(TEXT("[JNI] - AndroidThunkCpp_Iap_AcknowledgePurchase END"));
+		}
+	}
+
+	return bResult;
+}
+
 bool AndroidThunkCpp_Iap_QueryExistingPurchases()
 {
 	FPlatformMisc::LowLevelOutputDebugStringf(TEXT("[JNI] - AndroidThunkCpp_Iap_QueryExistingPurchases"));
diff --git a/Engine/Source/Runtime/Launch/Public/Android/AndroidJNI.h b/Engine/Source/Runtime/Launch/Public/Android/AndroidJNI.h
--- a/Engine/Source/Runtime/Launch/Public/Android/AndroidJNI.h
+++ b/Engine/Source/Runtime/Launch/Public/Android/AndroidJNI.h
@@ -118,6 +118,7 @@ public:
 	static jmethodID AndroidThunkJava_IapBeginPurchase;
 	static jmethodID AndroidThunkJava_IapIsAllowedToMakePurchases;
 	static jmethodID AndroidThunkJava_IapQueryExistingPurchases;
+	static jmethodID AndroidThunkJava_IapAcknowledgePurchase;
 	static jmethodID AndroidThunkJava_IapConsumePurchase;
 
 	// SurfaceView functionality for view scaling on some devices
