Fix Receipt management in OSSGooglePlay

Receipts retrieved using GetReceipts should be updated by calls to QueryReceipts. After a successful Checkout the new receipt is added to the existing known receipts

Ignored all asynchronous transaction updates received. Just trigger the UnexpectedPurchaseReceipt delegate in that case of Purchased transaction notifications

Active subscriptions and non consumed products will be present in the receipts received by GetReceipts after a succesfull call to QueryReceipts

diff --git a/Engine/Plugins/Online/Android/OnlineSubsystemGooglePlay/Source/Private/OnlinePurchaseGooglePlay.cpp b/Engine/Plugins/Online/Android/OnlineSubsystemGooglePlay/Source/Private/OnlinePurchaseGooglePlay.cpp
--- a/Engine/Plugins/Online/Android/OnlineSubsystemGooglePlay/Source/Private/OnlinePurchaseGooglePlay.cpp
+++ b/Engine/Plugins/Online/Android/OnlineSubsystemGooglePlay/Source/Private/OnlinePurchaseGooglePlay.cpp
@@ -356,8 +356,7 @@ void FOnlinePurchaseGooglePlay::QueryReceipts(const FUniqueNetId& UserId, bool /
 	if (!bQueryingReceipts)
 	{
 		/**
-		 * bRestoreReceipts is irrelevant because there is no concept of restore here, any purchase query will reveal any non consumed purchases. 
-		 * They will remain in the list until the the game consumes it via FinalizePurchase().
+		 * bRestoreReceipts is irrelevant because there is no concept of restore on GooglePlay, any purchase query will reveal all non consumed purchases. 
 		 */
 		bQueryingReceipts = true;
 		QueryReceiptsComplete = Delegate;
@@ -462,7 +461,6 @@ void FOnlinePurchaseGooglePlay::OnTransactionCompleteResponse(EGooglePlayBilling
 		{
 			KnownTransactions.Add(FinalReceipt);
 		}
-
 		InProgressTransaction->CheckoutCompleteDelegate.ExecuteIfBound(FinalResult, FinalReceipt);
 		InProgressTransaction.Reset();
 	}
@@ -470,29 +468,7 @@ void FOnlinePurchaseGooglePlay::OnTransactionCompleteResponse(EGooglePlayBilling
 	{
 		if (Result == EPurchaseTransactionState::Purchased)
 		{
-			TSharedRef<FPurchaseReceipt> NewReceipt = FOnlinePurchaseInProgressTransaction::GenerateReceipt(InTransactionData);;
-			TSharedRef<FPurchaseReceipt>* ExistingReceipt = KnownTransactions.FindByPredicate([&InTransactionData](const TSharedRef<FPurchaseReceipt>& Receipt)
-			{
-				return Receipt->TransactionId == InTransactionData.GetTransactionIdentifier();
-			});
-			
-			if(ExistingReceipt)
-			{
-				if((*ExistingReceipt)->TransactionState == EPurchaseTransactionState::Deferred)
-				{
-					UE_LOG_ONLINE_PURCHASE(Log, TEXT("Pending transaction completed"));
-				}
-				else
-				{
-					UE_LOG_ONLINE_PURCHASE(Log, TEXT("Unexpected transaction completed (a possible pending transaction finished from previous session?)"));
-				}
-				*ExistingReceipt = NewReceipt;
-			}
-			else
-			{
-				KnownTransactions.Add(NewReceipt);
-			}
-			
+			UE_LOG_ONLINE_PURCHASE(Log, TEXT("Unexpected transaction completed (maybe a pending transaction finished?)"));
 			TriggerOnUnexpectedPurchaseReceiptDelegates(*FUniqueNetIdGooglePlay::EmptyId());
 		}
 	}
@@ -532,28 +508,6 @@ void FOnlinePurchaseGooglePlay::OnQueryExistingPurchasesComplete(EGooglePlayBill
 	}
 }
 
-void FOnlinePurchaseGooglePlay::OnAcknowledgePurchaseComplete(EGooglePlayBillingResponseCode InResponseCode, const FString& TransactionIdentifier)
-{
-	UE_LOG_ONLINE_PURCHASE(Log, TEXT("FOnlinePurchaseGooglePlay::OnConsumePurchaseComplete Response: %s PurchaseToken: %s"), LexToString(InResponseCode), *TransactionIdentifier);
-}
-
-void FOnlinePurchaseGooglePlay::OnConsumePurchaseComplete(EGooglePlayBillingResponseCode InResponseCode, const FString& TransactionIdentifier)
-{
-	UE_LOG_ONLINE_PURCHASE(Log, TEXT("FOnlinePurchaseGooglePlay::OnConsumePurchaseComplete Response: %s PurchaseToken: %s"), LexToString(InResponseCode), *TransactionIdentifier);
-
-	if (InResponseCode == EGooglePlayBillingResponseCode::Ok || InResponseCode == EGooglePlayBillingResponseCode::ItemNotOwned)
-	{
-		KnownTransactions.RemoveAll([&TransactionIdentifier](const TSharedRef<FPurchaseReceipt>& Receipt) 
-		{
-			return Receipt->TransactionId == TransactionIdentifier;
-		});		
-	}
-	else
-	{
-		UE_LOG_ONLINE_PURCHASE(Log, TEXT("OnConsumePurchaseComplete failed"));
-	}
-}
-
 void FOnlinePurchaseGooglePlay::FinalizeReceiptValidationInfo(const FUniqueNetId& UserId, FString& InReceiptValidationInfo, const FOnFinalizeReceiptValidationInfoComplete& Delegate)
 {
 	FOnlineError DefaultSuccess(true);
@@ -634,42 +588,16 @@ JNI_METHOD void Java_com_epicgames_unreal_GooglePlayStoreHelper_NativePurchaseCo
 
 JNI_METHOD void Java_com_epicgames_unreal_GooglePlayStoreHelper_NativeConsumeComplete(JNIEnv* jenv, jobject /*Thiz*/, jint JavaResponseCode, jstring JavaPurchaseToken)
 {
-	FString PurchaseToken;
+	FString PurchaseToken = FJavaHelper::FStringFromParam(jenv, JavaPurchaseToken);;
 	EGooglePlayBillingResponseCode EGPResponse = (EGooglePlayBillingResponseCode)JavaResponseCode;
-
-	bool bWasSuccessful = (EGPResponse == EGooglePlayBillingResponseCode::Ok);
 	UE_LOG_ONLINE_PURCHASE(Verbose, TEXT("NativeConsumeComplete received response code: %s. PurchaseToken: %s\n"), LexToString(EGPResponse), *PurchaseToken);
-
-	PurchaseToken = FJavaHelper::FStringFromParam(jenv, JavaPurchaseToken);
-
-	if (auto OnlineSubGP = static_cast<FOnlineSubsystemGooglePlay* const>(IOnlineSubsystem::Get(GOOGLEPLAY_SUBSYSTEM)))
-	{
-		OnlineSubGP->ExecuteNextTick([OnlineSubGP, EGPResponse, PurchaseToken]()
-		{
-			TSharedPtr<FOnlinePurchaseGooglePlay> PurchaseInt = StaticCastSharedPtr<FOnlinePurchaseGooglePlay>(OnlineSubGP->GetPurchaseInterface());
-			PurchaseInt->OnConsumePurchaseComplete(EGPResponse, PurchaseToken);
-		});
-	}
 }
 
 JNI_METHOD void Java_com_epicgames_unreal_GooglePlayStoreHelper_NativeAcknowledgeComplete(JNIEnv* jenv, jobject /*Thiz*/, jint JavaResponseCode, jstring JavaPurchaseToken)
 {
-	FString PurchaseToken;
+	FString PurchaseToken = FJavaHelper::FStringFromParam(jenv, JavaPurchaseToken);;
 	EGooglePlayBillingResponseCode EGPResponse = (EGooglePlayBillingResponseCode)JavaResponseCode;
-
-	bool bWasSuccessful = (EGPResponse == EGooglePlayBillingResponseCode::Ok);
 	UE_LOG_ONLINE_PURCHASE(Verbose, TEXT("NativeAcknowledgeComplete received response code: %s. PurchaseToken: %s\n"), LexToString(EGPResponse), *PurchaseToken);
-
-	PurchaseToken = FJavaHelper::FStringFromParam(jenv, JavaPurchaseToken);
-
-	if (auto OnlineSubGP = static_cast<FOnlineSubsystemGooglePlay* const>(IOnlineSubsystem::Get(GOOGLEPLAY_SUBSYSTEM)))
-	{
-		OnlineSubGP->ExecuteNextTick([OnlineSubGP, EGPResponse, PurchaseToken]()
-		{
-			TSharedPtr<FOnlinePurchaseGooglePlay> PurchaseInt = StaticCastSharedPtr<FOnlinePurchaseGooglePlay>(OnlineSubGP->GetPurchaseInterface());
-			PurchaseInt->OnAcknowledgePurchaseComplete(EGPResponse, PurchaseToken);
-		});
-	}
 }
 
 #undef LOCTEXT_NAMESPACE
diff --git a/Engine/Plugins/Online/Android/OnlineSubsystemGooglePlay/Source/Public/OnlinePurchaseGooglePlay.h b/Engine/Plugins/Online/Android/OnlineSubsystemGooglePlay/Source/Public/OnlinePurchaseGooglePlay.h
--- a/Engine/Plugins/Online/Android/OnlineSubsystemGooglePlay/Source/Public/OnlinePurchaseGooglePlay.h
+++ b/Engine/Plugins/Online/Android/OnlineSubsystemGooglePlay/Source/Public/OnlinePurchaseGooglePlay.h
@@ -122,12 +122,6 @@ public:
 	/** Handle Java side query purchases completed notification */
 	void OnQueryExistingPurchasesComplete(EGooglePlayBillingResponseCode InResponseCode, const TArray<FGoogleTransactionData>& InExistingPurchases);
 
-	/** Handle Java side acknowledge purchase completed notification */
-	void OnAcknowledgePurchaseComplete(EGooglePlayBillingResponseCode InResponseCode, const FString& InPurchaseToken);
-
-	/** Handle Java side consume purchase completed notification */
-	void OnConsumePurchaseComplete(EGooglePlayBillingResponseCode InResponseCode, const FString& InPurchaseToken);
-
 	static bool IsSubscriptionProductId(const FString& ProductId);
 private:
 	
