Update GooglePlayBillingLibrary to v6

diff --git a/Engine/Plugins/Online/Android/OnlineSubsystemGooglePlay/Source/Java/BillingApiV2/com/epicgames/unreal/GooglePlayStoreHelper.java b/Engine/Plugins/Online/Android/OnlineSubsystemGooglePlay/Source/Java/BillingApiV2/com/epicgames/unreal/GooglePlayStoreHelper.java
--- a/Engine/Plugins/Online/Android/OnlineSubsystemGooglePlay/Source/Java/BillingApiV2/com/epicgames/unreal/GooglePlayStoreHelper.java
+++ b/Engine/Plugins/Online/Android/OnlineSubsystemGooglePlay/Source/Java/BillingApiV2/com/epicgames/unreal/GooglePlayStoreHelper.java
@@ -202,7 +202,7 @@ public class GooglePlayStoreHelper implements StoreHelper, PurchasesUpdatedListe
 			return false;
 		}
 
-		if (CommonProductType == BillingClient.ProductType.SUBS && ProductIds.length > 1)
+		if (CommonProductType.equals(BillingClient.ProductType.SUBS) && ProductIds.length > 1)
 		{
 			// This condition was added because the error produced by BillingClient.launchBillingFlow was not informative
 			Log.debug("[GooglePlayStoreHelper] - BeginPurchase - Failed! Purchasing multiple subscriptions at once is not supported");
@@ -393,6 +393,7 @@ public class GooglePlayStoreHelper implements StoreHelper, PurchasesUpdatedListe
 				@Override
 				public void onProductDetailsResponse(@NonNull BillingResult Result, @NonNull List<ProductDetails> ProductDetailsList)
 				{
+					int responseCode = Result.getResponseCode();
 					if (responseCode == BillingClient.BillingResponseCode.OK)
 					{
 						Set<String> SubscriptionProductIds = GetSubscriptionStoreProductIds(ProductDetailsList);
@@ -805,8 +806,6 @@ public class GooglePlayStoreHelper implements StoreHelper, PurchasesUpdatedListe
 				return "BillingResponseCode.USER_CANCELED";
 			case BillingClient.BillingResponseCode.SERVICE_UNAVAILABLE:
 				return "BillingResponseCode.SERVICE_UNAVAILABLE";
-			case BillingClient.BillingResponseCode.SERVICE_TIMEOUT:
-				return "BillingResponseCode.SERVICE_TIMEOUT";
 			case BillingClient.BillingResponseCode.SERVICE_DISCONNECTED:
 				return "BillingResponseCode.SERVICE_DISCONNECTED";
 			case BillingClient.BillingResponseCode.BILLING_UNAVAILABLE:
diff --git a/Engine/Plugins/Online/Android/OnlineSubsystemGooglePlay/Source/OnlineSubsystemGooglePlay_UPL.xml b/Engine/Plugins/Online/Android/OnlineSubsystemGooglePlay/Source/OnlineSubsystemGooglePlay_UPL.xml
--- a/Engine/Plugins/Online/Android/OnlineSubsystemGooglePlay/Source/OnlineSubsystemGooglePlay_UPL.xml
+++ b/Engine/Plugins/Online/Android/OnlineSubsystemGooglePlay/Source/OnlineSubsystemGooglePlay_UPL.xml
@@ -15,7 +15,7 @@
 	<buildGradleAdditions>
 		<insert>
 			dependencies {
-			implementation 'com.android.billingclient:billing:5.1.0'
+			implementation 'com.android.billingclient:billing:6.0.0'
 			}
 		</insert>
 	</buildGradleAdditions>
diff --git a/Engine/Plugins/Online/Android/OnlineSubsystemGooglePlay/Source/Private/OnlineStoreGooglePlayCommon.h b/Engine/Plugins/Online/Android/OnlineSubsystemGooglePlay/Source/Private/OnlineStoreGooglePlayCommon.h
--- a/Engine/Plugins/Online/Android/OnlineSubsystemGooglePlay/Source/Private/OnlineStoreGooglePlayCommon.h
+++ b/Engine/Plugins/Online/Android/OnlineSubsystemGooglePlay/Source/Private/OnlineStoreGooglePlayCommon.h
@@ -14,7 +14,6 @@ struct FGoogleTransactionData;
 enum class EGooglePlayBillingResponseCode : int8
 {
 	CustomLogicError = -127, // Custom value. Should match with GooglePlayStoreHelper.CustomLogicErrorResponse on GooglePlayStoreHelper.java
-	ServiceTimeout = -3,
 	FeatureNotSupported = -2,
 	ServiceDisconnected = -1,
 	Ok = 0,
@@ -42,8 +41,6 @@ inline const TCHAR* const LexToString(EGooglePlayBillingResponseCode InResponseC
 	{
 		case EGooglePlayBillingResponseCode::CustomLogicError:
 			return TEXT("CustomLogicError");
-		case EGooglePlayBillingResponseCode::ServiceTimeout:
-			return TEXT("ServiceTimeout");
 		case EGooglePlayBillingResponseCode::FeatureNotSupported:
 			return TEXT("FeatureNotSupported");
 		case EGooglePlayBillingResponseCode::ServiceDisconnected:
@@ -99,7 +96,6 @@ inline EInAppPurchaseState::Type ConvertGPResponseCodeToIAPState(const EGooglePl
 		case EGooglePlayBillingResponseCode::ItemNotOwned:
 			return EInAppPurchaseState::NotAllowed;
 		case EGooglePlayBillingResponseCode::CustomLogicError:
-		case EGooglePlayBillingResponseCode::ServiceTimeout:
 		case EGooglePlayBillingResponseCode::FeatureNotSupported:
 		case EGooglePlayBillingResponseCode::ServiceDisconnected:
 		case EGooglePlayBillingResponseCode::ServiceUnavailable:
@@ -139,7 +135,6 @@ inline EPurchaseTransactionState ConvertGPResponseToPurchaseTransactionState(con
 		case EGooglePlayBillingResponseCode::ItemNotOwned:
 			return EPurchaseTransactionState::Invalid;
 		case EGooglePlayBillingResponseCode::CustomLogicError:
-		case EGooglePlayBillingResponseCode::ServiceTimeout:
 		case EGooglePlayBillingResponseCode::FeatureNotSupported:
 		case EGooglePlayBillingResponseCode::ServiceDisconnected:
 		case EGooglePlayBillingResponseCode::ServiceUnavailable:
