Update GooglePlay games dependencies. This includes
 - Create our own JNI wrapper on top of latest play-services-games-v2 classes
 - Adapt OSSGooglePlay to use this wrapper
 - Remove deprecated/unsupported functionality
 - Fix some issues with JNI support

diff --git a/Engine/Build/Android/Java/aar-imports.txt b/Engine/Build/Android/Java/aar-imports.txt
--- a/Engine/Build/Android/Java/aar-imports.txt
+++ b/Engine/Build/Android/Java/aar-imports.txt
@@ -1,6 +1,3 @@
 repositories $(ANDROID_HOME)/extras
-repositories $(ENGINEDIR)/Source/ThirdParty/Android/extras
-com.google.android.gms,play-services-auth,18.1.0
-com.google.android.gms,play-services-nearby,18.1.0
-com.google.android.gms,play-services-games,18.0.1
-com.google.android.gms,play-services-plus,17.0.0
+com.google.android.gms,play-services-auth,20.7.0
+com.google.android.gms,play-services-games-v2,19.0.0
diff --git a/Engine/Build/Android/Java/src/com/epicgames/unreal/GameActivity.java.template b/Engine/Build/Android/Java/src/com/epicgames/unreal/GameActivity.java.template
--- a/Engine/Build/Android/Java/src/com/epicgames/unreal/GameActivity.java.template
+++ b/Engine/Build/Android/Java/src/com/epicgames/unreal/GameActivity.java.template
@@ -168,8 +168,6 @@ import com.google.android.gms.common.GooglePlayServicesUtil;
 import com.google.android.gms.common.ConnectionResult;
 import com.google.android.gms.games.Games;
 
-import com.google.android.gms.plus.Plus;
-
 import com.google.vr.sdk.samples.permission.PermissionHelper;
 
 import java.nio.charset.StandardCharsets;
@@ -240,8 +238,6 @@ import com.epicgames.unreal.psoservices.PSOProgramServiceAccessor;
 //  to them in this class.
 
 public class GameActivity extends $${gameActivitySuperClass}$$ implements SurfaceHolder.Callback2,
-															GoogleApiClient.ConnectionCallbacks,
-															GoogleApiClient.OnConnectionFailedListener,
 															SensorEventListener,
 															Logger.ILoggerCallback,
 //$${gameActivityImplementsAdditions}$$
@@ -431,8 +427,6 @@ public class GameActivity extends $${gameActivitySuperClass}$$ implements Surfac
 	/** AssetManger reference - populated on start up and used when the OBB is packed into the APK */
 	protected AssetManager			AssetManagerReference;
 	
-	private GoogleApiClient googleClient;
-
 	// layout required by popups, e.g ads, native controls
 	LinearLayout activityLayout;
 
@@ -4566,142 +4560,6 @@ public class GameActivity extends $${gameActivitySuperClass}$$ implements Surfac
 		}
 	}
 
-	public void AndroidThunkJava_ResetAchievements()
-	{
-		/* Disable so don't need GET_ACCOUNTS - we don't need this
-		try
-        {
-			String accesstoken = GetAccessToken();
-			          
-			if(!accesstoken.equals(""))
-			{
-				String ResetURL = "https://www.googleapis.com/games/v1management/achievements/reset?access_token=" + accesstoken;
-				Log.debug("AndroidThunkJava_ResetAchievements: using URL " + ResetURL);
-
-				URL url = new URL(ResetURL);
-				HttpURLConnection urlConnection = (HttpURLConnection)url.openConnection();
-
-				try
-				{
-					urlConnection.setRequestMethod("POST");
-					int status = urlConnection.getResponseCode();
-					Log.debug("AndroidThunkJava_ResetAchievements: HTTP response is " + status);
-				}
-				finally
-				{
-					urlConnection.disconnect();
-				}
-			}
-			else
-			{
-				Log.debug("AndroidThunkJava_ResetAchievements: Access Token not returned.  Possible reason is that android.permission.GET_ACCOUNTS is not granted.  Make sure to add in by going to Project settings > Android > Advanced Project Settings and check the box for \"Request Access Token On Connect\". ");
-			}
-        }
-        catch(Exception e)
-        {
-            Log.debug("AndroidThunkJava_ResetAchievements failed: " + e.getMessage());
-        }
-		*/
-		Log.debug("AndroidThunkJava_ResetAchievements: disabled");
-	}
-
-	// TODO: replace this with non-depreciated method (OK now for up to API-25)
-	@TargetApi(23)
-	private String GetAccessToken()
-	{
-		String accesstoken = "";
-
-		try
-        {
-			if (PermissionHelper.checkPermission("android.permission.GET_ACCOUNTS") && googleClient != null)
-			{
-				String email = Plus.AccountApi.getAccountName(googleClient);
-				Log.debug("GetAccessToken: using email " + email);
-
-				accesstoken = GoogleAuthUtil.getToken(this, email, "oauth2:https://www.googleapis.com/auth/games");
-			}
-		}
-		catch(Exception e)
-        {
-            Log.debug("GetAccessToken failed: " + e.getMessage());
-        }
-
-		return accesstoken;
-	}
-
-	public void AndroidThunkJava_GoogleClientConnect()
-	{
-		if (googleClient != null)
-		{
-			googleClient.connect();
-		}
-	}
-
-	public void AndroidThunkJava_GoogleClientDisconnect()
-	{
-		if (googleClient != null)
-		{
-			googleClient.disconnect();
-		}
-	}
-
-	// Google Client connected successfully
-	@Override
-	public void onConnected(Bundle connectionHint)
-	{
-		if (googleClient != null && googleClient.isConnected())
-		{
-			new Thread(new Runnable()
-			{
-				public void run()
-				{
-					/* Don't do this since deprecated and don't want to request GET_ACCOUNTS permission
-					try
-					{
-						String accesstoken = GetAccessToken();
-						if(!accesstoken.equals(""))
-						{
-							Log.debug("Google Client Connect using Access Token " + accesstoken);
-							nativeGoogleClientConnectCompleted(true, accesstoken);
-						}
-						else
-						{
-							Log.debug("Google Client Connect succeeded but no access token returned");
-							nativeGoogleClientConnectCompleted(true, "");
-						}
-					}
-					catch (Exception e)
-					{
-						Log.debug("Google Client Connect failed: " + e.getMessage());
-
-						nativeGoogleClientConnectCompleted(false, "");
-					}
-					*/
-					nativeGoogleClientConnectCompleted(true, "NOT_ACQUIRED");
-				}
-			}).start();
-		}
-		else
-		{
-			nativeGoogleClientConnectCompleted(false, "");
-		}
-	}
-
-	// Google Client connection failed
-	@Override
-	public void onConnectionFailed(ConnectionResult connectionResult)
-	{
-		Log.debug("Google Client Connect failed. Error Code: " + connectionResult.getErrorCode() + " Description: " + connectionResult.getErrorMessage());
-		nativeGoogleClientConnectCompleted(false, "");
-	}
-
-	// Google Client connection suspended
-	@Override
-	public void onConnectionSuspended(int cause)
-	{
-		Log.debug("Google Client Connect Suspended: " + cause);
-	}
-
 	public AssetManager AndroidThunkJava_GetAssetManager()
 	{
 		if(AssetManagerReference == null)
@@ -7080,8 +6938,6 @@ public class GameActivity extends $${gameActivitySuperClass}$$ implements Surfac
 
 	public native void nativeOnActivityResult(GameActivity activity, int requestCode, int resultCode, Intent data);
 
-	public native void nativeGoogleClientConnectCompleted(boolean bSuccess, String accessToken);
-
 	public native void nativeVirtualKeyboardShown(int left, int top, int right, int bottom);
 	public native void nativeVirtualKeyboardVisible(boolean bShown);
 
diff --git a/Engine/Plugins/Online/Android/OnlineSubsystemGooglePlay/Source/Java/BillingApiV2/com/epicgames/unreal/GooglePlayStoreHelper.java b/Engine/Plugins/Online/Android/OnlineSubsystemGooglePlay/Source/Java/BillingApiV2/com/epicgames/unreal/GooglePlayStoreHelper.java
deleted file mode 100644
diff --git a/Engine/Plugins/Online/Android/OnlineSubsystemGooglePlay/Source/Java/com/epicgames/unreal/GooglePlayGamesWrapper.java b/Engine/Plugins/Online/Android/OnlineSubsystemGooglePlay/Source/Java/com/epicgames/unreal/GooglePlayGamesWrapper.java
new file mode 100644
--- /dev/null
+++ b/Engine/Plugins/Online/Android/OnlineSubsystemGooglePlay/Source/Java/com/epicgames/unreal/GooglePlayGamesWrapper.java
@@ -0,0 +1,363 @@
+// Copyright Epic Games, Inc. All Rights Reserved.
+
+package com.epicgames.unreal;
+
+import android.app.Activity;
+import android.content.Context;
+
+import androidx.annotation.NonNull;
+
+import com.google.android.gms.games.AchievementsClient;
+import com.google.android.gms.games.GamesSignInClient;
+import com.google.android.gms.games.LeaderboardsClient;
+import com.google.android.gms.games.achievement.Achievement;
+import com.google.android.gms.games.achievement.AchievementBuffer;
+import com.google.android.gms.games.leaderboard.LeaderboardScore;
+import com.google.android.gms.games.leaderboard.LeaderboardVariant;
+import com.google.android.gms.games.PlayGames;
+import com.google.android.gms.games.PlayGamesSdk;
+import com.google.android.gms.games.Player;
+import com.google.android.gms.tasks.Task;
+import com.google.android.gms.tasks.Tasks;
+
+import java.util.ArrayList;
+import java.util.List;
+
+public class GooglePlayGamesWrapper
+{
+	public static Logger Log = new Logger("UE", "GooglePlayGamesWrapper");
+
+	private static class AuthCodeSettings
+	{
+		public AuthCodeSettings(String InServerClientId, boolean InForceRefreshToken)
+		{
+			ServerClientId = InServerClientId;
+			ForceRefreshToken = InForceRefreshToken;
+		}
+
+		public String ServerClientId;
+		public boolean ForceRefreshToken;
+	}
+
+	// This method should be called from the onCreate method on the Application
+	static public void Initialize(Context AppContext)
+	{
+		Log.debug("Initializing");
+		PlayGamesSdk.initialize(AppContext);
+	}
+
+	// Login method to be accessed through JNI
+	public static void Login(long TaskPtr, final Activity Activity, String InServerClientId, boolean InForceRefreshToken)
+	{
+		Log.debug("Login started");
+		GamesSignInClient SignInClient = PlayGames.getGamesSignInClient(Activity);
+
+		SignInClient.isAuthenticated().addOnCompleteListener(isAuthenticatedTask ->
+			{
+				if (isAuthenticatedTask.isSuccessful())
+				{
+					final AuthCodeSettings AuthCodeData = (InServerClientId != null) ? new AuthCodeSettings(InServerClientId, InForceRefreshToken) : null;
+					if (isAuthenticatedTask.getResult().isAuthenticated())
+					{
+						Log.debug("User already authenticated. Requesting player data");
+						RequestIdentityData(TaskPtr, Activity, AuthCodeData);
+					}
+					else
+					{
+						Log.debug("User not yet authenticated. Attempting sign in");
+						SignInClient.signIn().addOnCompleteListener(SignInTask ->
+							{
+								if (SignInTask.isSuccessful())
+								{
+									if (SignInTask.getResult().isAuthenticated())
+									{
+										Log.debug("Sign in succeeded. Requesting player data");
+										RequestIdentityData(TaskPtr, Activity, AuthCodeData);
+									}
+									else
+									{
+										LoginFailed(TaskPtr, "Sign in failed");
+									}
+								}
+								else
+								{
+									LoginFailed(TaskPtr, "GamesSignInClient.signIn task failed. Maybe service was interrupted or wrong configuration?");
+								}
+							});
+					}
+				}
+				else
+				{
+					LoginFailed(TaskPtr, "GamesSignInClient.isAuthenticated task failed. Maybe service was interrupted or wrong configuration?");
+				}
+			});
+	}
+
+	private static Task<String> TryRequestAuthCode(Activity Activity, AuthCodeSettings Settings)
+	{
+		if (Settings == null)
+		{
+			return null;
+		}
+		Log.debug("Will request auth code");
+
+		GamesSignInClient SignInClient = PlayGames.getGamesSignInClient(Activity);
+		return SignInClient.requestServerSideAccess(Settings.ServerClientId, Settings.ForceRefreshToken);
+	}
+
+	private static void RequestIdentityData(long TaskPtr, Activity Activity, AuthCodeSettings Settings)
+	{
+		List<Task<?>> LoginTasks = new ArrayList<>();
+
+		final Task<Player> GetPlayerTask = PlayGames.getPlayersClient(Activity).getCurrentPlayer();
+		LoginTasks.add(GetPlayerTask);
+
+		final Task<String> GetAuthCodeTask = TryRequestAuthCode(Activity, Settings);
+		if (GetAuthCodeTask != null)
+		{
+			LoginTasks.add(GetAuthCodeTask);
+		}
+
+		Tasks.whenAllComplete(LoginTasks).addOnCompleteListener(CompositeTask ->
+			{
+				// Task returned by whenAllComplete always succeed
+				if (GetPlayerTask.isSuccessful())
+				{
+					Player CurrentPlayer = GetPlayerTask.getResult();
+					String AuthCode = null;
+					if (GetAuthCodeTask != null)
+					{
+						if (GetAuthCodeTask.isSuccessful())
+						{
+							AuthCode = GetAuthCodeTask.getResult();
+						}
+						else
+						{
+							Log.warn("Could not get auth code. Maybe service was interrupted or wrong configuration?");
+						}
+					}
+					LoginSuccess(TaskPtr, CurrentPlayer, AuthCode);
+				}
+				else
+				{
+					LoginFailed(TaskPtr, "PlayersClient.getCurrentPlayer task failed. Maybe service was interrupted or wrong configuration?");
+				}
+			});
+	}
+
+	static void LoginSuccess(long TaskPtr, @NonNull Player Player, String AuthCode)
+	{
+		Log.debug("Login succeeded");
+		nativeLoginSuccess(TaskPtr, Player.getPlayerId(), Player.getDisplayName(), AuthCode);
+	}
+
+	static void LoginFailed(long TaskPtr, @NonNull String ErrorMessage)
+	{
+		Log.warn(ErrorMessage);
+		nativeLoginFailed(TaskPtr);
+	}
+
+	public static void RequestPlayerLeaderboardScore(long TaskPtr, Activity Activity, String LeaderboardId)
+	{
+		PlayGames.getLeaderboardsClient(Activity).loadCurrentPlayerLeaderboardScore(LeaderboardId, LeaderboardVariant.TIME_SPAN_ALL_TIME,  LeaderboardVariant.COLLECTION_PUBLIC)
+			.addOnCompleteListener(LoadScoreTask ->
+			{
+				if (LoadScoreTask.isSuccessful())
+				{
+					Log.debug("Score received for leaderboard" + LeaderboardId);
+					LeaderboardScore Result = LoadScoreTask.getResult().get();
+
+					if (Result == null)
+					{
+						Log.warn("Player not yet in the leaderboard " + LeaderboardId);
+					}
+					else
+					{
+						nativeLeaderboardRequestSuccess(TaskPtr,
+								Result.getScoreHolder().getDisplayName(),
+								Result.getScoreHolder().getPlayerId(),
+								Result.getRank(),
+								Result.getRawScore());
+						return;
+					}
+				}
+				nativeLeaderboardRequestFailed(TaskPtr);
+			}); 
+	}
+
+	public static void SubmitLeaderboardsScores(long TaskPtr, Activity Activity, String[] Leaderboards, long[] NewScores)
+	{
+		if (Leaderboards.length != NewScores.length)
+		{
+			Log.warn("Missmatch in leaderboards/scores array sizes");
+		}
+
+		LeaderboardsClient Client = PlayGames.getLeaderboardsClient(Activity); 
+		int Count = Math.min(Leaderboards.length, NewScores.length);
+		List<Task<?>> UpdateLeaderboardTasks = new ArrayList<>();
+
+		for( int Idx = 0; Idx < Count; Idx ++)
+		{
+			UpdateLeaderboardTasks.add(Client.submitScoreImmediate(Leaderboards[Idx], NewScores[Idx]));
+		}
+
+		Tasks.whenAllComplete(UpdateLeaderboardTasks).addOnCompleteListener(CompositeTask ->
+			{
+				Log.debug("Submit leaderboards scores response received");
+
+				List<Task<?>> Tasks = CompositeTask.getResult();
+				boolean bSuccess = true;
+				for (int Idx = 0; Idx < Count; Idx ++)
+				{
+					Task<?> Task = Tasks.get(Idx);
+					if (!Task.isSuccessful())
+					{
+						bSuccess = false;
+						Log.warn("Failed to update leaderboard " + Leaderboards[Idx] + ":" + Task.getException().getMessage());
+					}
+				}
+
+				nativeFlushLeaderboardsCompleted(TaskPtr, bSuccess);
+			});
+	}
+
+	public static void ShowLeaderboardUI(Activity Activity, String LeaderboardId)
+	{
+		PlayGames.getLeaderboardsClient(Activity)
+			.getLeaderboardIntent(LeaderboardId)
+			.addOnSuccessListener(Intent -> Activity.startActivityForResult(Intent, GameActivity.REQUEST_CODE_SHOW_GOOGLEPLAY_UI));
+	}
+
+	public static void QueryAchievements(long TaskPtr, Activity Activity)
+	{
+		PlayGames.getAchievementsClient(Activity)
+			.load(false)
+			.addOnCompleteListener(LoadAchievementsTask ->
+			{
+				if (LoadAchievementsTask.isSuccessful())
+				{
+					Log.debug("Achievements data received");
+
+					AchievementBuffer Achievements = LoadAchievementsTask.getResult().get();
+					int Count = Achievements.getCount();
+
+					String[] AchievementIds = new String[Count];
+					int[] Types = new int[Count];
+					int[] Steps = new int[Count];
+					int[] TotalSteps = new int[Count];
+					String[] Titles = new String[Count];
+					String[] Descriptions = new String[Count];
+					boolean[] IsHidden = new boolean[Count];
+					long[] LastUpdatedTimestamps = new long[Count];
+
+					for (int Idx = 0; Idx < Count; Idx++)
+					{
+						Achievement Current = Achievements.get(Idx);
+
+						AchievementIds[Idx] = Current.getAchievementId(); 
+						Types[Idx] = Current.getType();
+						Titles[Idx] = Current.getName();
+						Descriptions[Idx] = Current.getDescription();
+						IsHidden[Idx] = Current.getState() == Achievement.STATE_HIDDEN;
+						LastUpdatedTimestamps[Idx] = Current.getLastUpdatedTimestamp();
+
+						if (Current.getType() == Achievement.TYPE_INCREMENTAL)
+						{
+							Steps[Idx] = Current.getCurrentSteps();
+							TotalSteps[Idx] = Current.getTotalSteps();
+						}
+						else //Achievement.TYPE_STANDARD
+						{
+							TotalSteps[Idx] = 0;
+							if (Current.getState() == Achievement.STATE_UNLOCKED)
+							{
+								Steps[Idx] = 1;
+							}
+							else
+							{
+								Steps[Idx] = 0;
+							}
+						}
+					}
+
+					nativeQueryAchievementsSuccess(TaskPtr, AchievementIds, Types, Steps, TotalSteps, Titles, Descriptions, IsHidden, LastUpdatedTimestamps);
+				}
+				else
+				{
+					Log.warn("Failed to get achievements data: " + LoadAchievementsTask.getException().getMessage());
+					nativeQueryAchievementsFailed(TaskPtr);
+				}
+			});
+	}
+
+	static final int WRITE_ACHIEVEMENT_ACTION_UNLOCK = 0;
+	static final int WRITE_ACHIEVEMENT_ACTION_WRITE_STEPS = 1;
+	
+	public static void WriteAchievements(long TaskPtr, Activity Activity, String[] AchievementIds, int[] Actions, int[] Values)
+	{
+		if (AchievementIds.length != Actions.length || AchievementIds.length != Values.length)
+		{
+			Log.warn("Missmatch in leaderboards/scores array sizes");
+		}
+
+		AchievementsClient Client = PlayGames.getAchievementsClient(Activity);
+		List<Task<?>> WriteTasks = new ArrayList<>();
+
+		int Count = Math.min(Math.min(AchievementIds.length, Actions.length), Values.length);
+		for (int Idx = 0; Idx < Count; Idx ++)
+		{
+			switch ( Actions[Idx] )
+			{
+				case WRITE_ACHIEVEMENT_ACTION_UNLOCK:
+				{
+					WriteTasks.add(Client.unlockImmediate(AchievementIds[Idx]));
+				}
+				break;
+				case WRITE_ACHIEVEMENT_ACTION_WRITE_STEPS:
+				{
+					WriteTasks.add(Client.setStepsImmediate(AchievementIds[Idx], Values[Idx]));
+				}
+				break;
+			}
+		}
+		Tasks.whenAllComplete(WriteTasks).addOnCompleteListener(CompositeTask ->
+			{
+				Log.debug("Write achievements response received");
+
+				List<Task<?>> Tasks = CompositeTask.getResult();
+				List<String> Succeeded = new ArrayList<>();
+				for (int Idx = 0; Idx < Count; Idx ++)
+				{
+					Task<?> Task = Tasks.get(Idx);
+					if (Task.isSuccessful())
+					{
+						Succeeded.add(AchievementIds[Idx]);
+					}
+					else
+					{
+						Log.warn("Failed to write achievement" + AchievementIds[Idx] + ":" + Task.getException().getMessage());
+					}
+				}
+				nativeWriteAchievementsCompleted(TaskPtr, Succeeded.toArray(new String[0]));
+			});
+	}
+
+	public static void ShowAchievementsUI(Activity Activity)
+	{
+		PlayGames.getAchievementsClient(Activity)
+			.getAchievementsIntent()
+			.addOnSuccessListener(Intent -> Activity.startActivityForResult(Intent, GameActivity.REQUEST_CODE_SHOW_GOOGLEPLAY_UI));
+	}
+
+	native static void nativeLoginSuccess(long TaskPtr, String PlayerId, String DisplayName, String AuthCode);
+	native static void nativeLoginFailed(long TaskPtr);
+
+	native static void nativeLeaderboardRequestSuccess(long TaskPtr, String DisplayName, String PlayerId, long Rank, long RawScore);
+	native static void nativeLeaderboardRequestFailed(long TaskPtr);
+	native static void nativeFlushLeaderboardsCompleted(long TaskPtr, boolean Success);
+
+	native static void nativeQueryAchievementsSuccess(long TaskPtr, String[] AchievementIds, int[] Types, int[] Steps, int[] TotalSteps, String[] Titles, String[] Descriptions, boolean[] IsHidden, long[] LastUpdatedTimestamps);
+	native static void nativeQueryAchievementsFailed(long TaskPtr);
+
+	native static void nativeWriteAchievementsCompleted(long TaskPtrr, String[] Succeeded);
+}
diff --git a/Engine/Plugins/Online/Android/OnlineSubsystemGooglePlay/Source/Java/com/epicgames/unreal/GooglePlayStoreHelper.java b/Engine/Plugins/Online/Android/OnlineSubsystemGooglePlay/Source/Java/com/epicgames/unreal/GooglePlayStoreHelper.java
new file mode 100644
--- /dev/null
+++ b/Engine/Plugins/Online/Android/OnlineSubsystemGooglePlay/Source/Java/com/epicgames/unreal/GooglePlayStoreHelper.java
@@ -0,0 +1,868 @@
+// Copyright Epic Games, Inc. All Rights Reserved.
+
+package com.epicgames.unreal;
+
+import android.util.Pair;
+
+import androidx.annotation.NonNull;
+import androidx.annotation.Nullable;
+
+import com.android.billingclient.api.BillingClient;
+import com.android.billingclient.api.BillingClientStateListener;
+import com.android.billingclient.api.BillingFlowParams;
+import com.android.billingclient.api.BillingResult;
+import com.android.billingclient.api.AcknowledgePurchaseParams;
+import com.android.billingclient.api.AcknowledgePurchaseResponseListener;
+import com.android.billingclient.api.ConsumeParams;
+import com.android.billingclient.api.ConsumeResponseListener;
+import com.android.billingclient.api.ProductDetails;
+import com.android.billingclient.api.ProductDetailsResponseListener;
+import com.android.billingclient.api.Purchase;
+import com.android.billingclient.api.PurchasesResponseListener;
+import com.android.billingclient.api.PurchasesUpdatedListener;
+import com.android.billingclient.api.QueryProductDetailsParams;
+import com.android.billingclient.api.QueryPurchasesParams;
+import com.android.vending.billing.util.Base64; //WMM should use different base64 here.
+
+import java.util.ArrayList;
+import java.util.Arrays;
+import java.util.HashSet;
+import java.util.List;
+import java.util.Objects;
+import java.util.Set;
+import java.util.stream.Collectors;
+
+public class GooglePlayStoreHelper implements StoreHelper, PurchasesUpdatedListener
+{
+	// Our IAB helper interface provided by google.
+	private final BillingClient Client;
+
+	// Flag that determines whether the store is ready to use.
+	private boolean bIsIapSetup = false;
+
+	// Output device for log messages.
+	private final Logger Log;
+
+	// Cache access to the games activity.
+	private final GameActivity MainActivity;
+
+	public GooglePlayStoreHelper(GameActivity InGameActivity, Logger InLog)
+	{
+		Log = InLog;
+		Log.debug("[GooglePlayStoreHelper] - GooglePlayStoreHelper");
+
+		MainActivity = InGameActivity;
+
+		Client = BillingClient.newBuilder(MainActivity)
+			.setListener(this)
+			.enablePendingPurchases()
+			.build();
+
+		Client.startConnection(new BillingClientStateListener() {
+			@Override
+			public void onBillingSetupFinished(@NonNull BillingResult billingResult)
+			{
+				if (billingResult.getResponseCode() == BillingClient.BillingResponseCode.OK)
+				{
+					bIsIapSetup = true;
+					Log.debug("[GooglePlayStoreHelper] - In-app billing supported for " + MainActivity.getPackageName());
+				}
+				else
+				{
+					Log.debug("[GooglePlayStoreHelper] - In-app billing NOT supported for " + MainActivity.getPackageName() + " error " + TranslateBillingResult(billingResult));
+				}
+			}
+
+			@Override
+			public void onBillingServiceDisconnected() {
+				// Try to restart the connection on the next request to
+				// Google Play by calling the startConnection() method.
+				bIsIapSetup = false;
+			}
+		});
+	}
+
+	///////////////////////////////////////////////////////
+	// The StoreHelper interfaces implementation for Google Play Store.
+
+	/**
+	 * Determine whether the store is ready for purchases.
+	 */
+	@Override
+	public boolean IsAllowedToMakePurchases()
+	{
+		Log.debug("[GooglePlayStoreHelper] - IsAllowedToMakePurchases. Billing client connected: " + bIsIapSetup);
+		return bIsIapSetup;
+	}
+
+	/**
+	 * Gets a set of store product ids for all the subscriptions provided in knownProducts.
+	 */
+	private Set<String> GetSubscriptionStoreProductIds(List<ProductDetails> KnownProducts)
+	{
+		Set<String> SubscriptionIds = new HashSet<String>();
+		for(ProductDetails Details: KnownProducts)
+		{
+			if(Details.getProductType().equals(BillingClient.ProductType.SUBS))
+			{
+				SubscriptionIds.add(Details.getProductId());
+			}
+		}
+		return SubscriptionIds;
+	}
+
+	/**
+	 * Query product details for the provided product ids
+	 */
+	@Override
+	public boolean QueryInAppPurchases(String[] InProductIDs)
+	{
+		Log.debug("[GooglePlayStoreHelper] - QueryInAppPurchases");
+		Log.debug("[GooglePlayStoreHelper] - NumProducts: " + InProductIDs.length);
+		Log.debug("[GooglePlayStoreHelper] - Products: " + Arrays.toString(InProductIDs));
+
+		if (InProductIDs.length > 0)
+		{
+			QueryProductDetailsParams QueryProducts = CreateQueryProductDetailsParams(InProductIDs, BillingClient.ProductType.INAPP);
+			QueryProductDetailsParams QuerySubs = CreateQueryProductDetailsParams(InProductIDs, BillingClient.ProductType.SUBS);
+
+			ProductDetailsResponseListener AggregatedListener = new ProductDetailsResponseListener() {
+				// This method is marked synchronized because although Google Billing Library documentation says "All methods annotated with AnyThread can be 
+				// called from any thread and all the asynchronous callbacks will be returned on the same thread" those callbacks may be called on different 
+				// threads if invoked from JNI.
+				@Override
+				public synchronized void onProductDetailsResponse(@NonNull BillingResult Result, @NonNull List<ProductDetails> ProductDetailsList)
+				{
+					if(AwaitingResponseCount <= 0)
+					{
+						return;
+					}
+					int ResponseCode = Result.getResponseCode();
+					if(ResponseCode == BillingClient.BillingResponseCode.OK)
+					{
+						AllProductDetails.addAll(ProductDetailsList);
+						AwaitingResponseCount -= 1;
+						Log.debug("[GooglePlayStoreHelper] - QueryInAppPurchases - Received new " + ProductDetailsList.size() + " products. Waiting for additional responses: " + (AwaitingResponseCount > 0));
+						if (AwaitingResponseCount == 0)
+						{
+							Log.debug("[GooglePlayStoreHelper] - QueryInAppPurchases - ProductDetails:" + AllProductDetails.toString());
+							Log.debug("[GooglePlayStoreHelper] - QueryInAppPurchases " + AllProductDetails.size() + " items - Success!");
+
+							ListOfProductDetailsAsArrays DetailArrays = new ListOfProductDetailsAsArrays(AllProductDetails);
+							NativeQueryComplete(BillingClient.BillingResponseCode.OK, DetailArrays.ProductIds, DetailArrays.Titles, DetailArrays.Descriptions, DetailArrays.Prices, DetailArrays.PricesRaw, DetailArrays.CurrencyCodes);
+						}
+					}
+					else
+					{
+						AwaitingResponseCount = 0;
+						Log.debug("[GooglePlayStoreHelper] - QueryInAppPurchases - Failed with: " + TranslateBillingResult(Result));
+						NativeQueryComplete(ResponseCode, null, null, null, null, null, null);
+					}
+
+					if (AwaitingResponseCount == 0)
+					{
+						Log.debug("[GooglePlayStoreHelper] - NativeQueryComplete done!");
+					}
+				}
+
+				// We will be querying for in app products and subscriptions 
+				int AwaitingResponseCount = 2;
+				List<ProductDetails> AllProductDetails = new ArrayList<>();
+			};
+
+			Client.queryProductDetailsAsync(QueryProducts, AggregatedListener);
+			Client.queryProductDetailsAsync(QuerySubs, AggregatedListener);
+		}
+		else
+		{
+			Log.debug("[GooglePlayStoreHelper] - no products given to query");
+			return false;
+		}
+
+		return true;
+	}
+
+	/**
+	 * Start the purchase flow for a particular product id
+	 */
+	@Override
+	public boolean BeginPurchase(final String[] ProductIds, final String ObfuscatedAccountId)
+	{
+		if (ProductIds.length == 0)
+		{
+			Log.debug("[GooglePlayStoreHelper] - BeginPurchase - Failed! No product ids provided");
+			return false;
+
+		}
+
+		String CommonProductType = EncodedProductId.GetCommonProductType(ProductIds);
+		if (CommonProductType == null)
+		{
+			Log.debug("[GooglePlayStoreHelper] - BeginPurchase - Failed! All products should have the same product type");
+			return false;
+		}
+
+		if (CommonProductType.equals(BillingClient.ProductType.SUBS) && ProductIds.length > 1)
+		{
+			// This condition was added because the error produced by BillingClient.launchBillingFlow was not informative
+			Log.debug("[GooglePlayStoreHelper] - BeginPurchase - Failed! Purchasing multiple subscriptions at once is not supported");
+			return false;
+		}
+
+		Log.debug("[GooglePlayStoreHelper] - BeginPurchase - querying product information " + String.join(",", ProductIds));
+		QueryProductDetailsParams Params = CreateQueryProductDetailsParams(ProductIds, CommonProductType);
+		Client.queryProductDetailsAsync(Params, new ProductDetailsResponseListener() {
+			@Override
+			public void onProductDetailsResponse(@NonNull BillingResult Result, @NonNull List<ProductDetails> ProductDetailsList)
+			{
+				if(Result.getResponseCode() == BillingClient.BillingResponseCode.OK)
+				{
+					Log.debug("[GooglePlayStoreHelper] - BeginPurchase - CreateBillingFlowParams with " + ProductDetailsList.stream().map(ProductDetails::getName).collect(Collectors.joining(",")));
+					final BillingFlowParams Params = CreateBillingFlowParams(ProductIds, ProductDetailsList, ObfuscatedAccountId);
+					if(Params == null)
+					{
+						Log.debug("[GooglePlayStoreHelper] - BeginPurchase - Failed! Matching products may not exist in the store");
+						NativePurchaseComplete(BillingClient.BillingResponseCode.DEVELOPER_ERROR, ProductIds, Purchase.PurchaseState.UNSPECIFIED_STATE,"", "", "");
+						return;
+					}
+					// BillingClient.launchBillingFlow is annotated @UiThread, so dispatch the call to the Activity UI thread
+					MainActivity.runOnUiThread(new Runnable()
+					{
+						@Override
+						public void run() {
+							Log.debug("[GooglePlayStoreHelper] - BeginPurchase - launchBillingFlow " + ProductDetailsList.stream().map(ProductDetails::getName).collect(Collectors.joining(",")));
+							BillingResult Result = Client.launchBillingFlow(MainActivity, Params);
+							int ResponseCode = Result.getResponseCode();
+							if (ResponseCode != BillingClient.BillingResponseCode.OK)
+							{
+								Log.debug("[GooglePlayStoreHelper] - BeginPurchase - Failed! " + TranslateBillingResult(Result));
+								NativePurchaseComplete(ResponseCode, ProductIds, Purchase.PurchaseState.UNSPECIFIED_STATE, "", "", "");
+							}
+						}
+					});
+				}
+				else
+				{
+					Log.debug("[GooglePlayStoreHelper] - BeginPurchase - Failed! " + TranslateBillingResult(Result));
+					NativePurchaseComplete(Result.getResponseCode(), ProductIds, Purchase.PurchaseState.UNSPECIFIED_STATE,"", "", "");
+				}
+			}
+		});
+		return true;
+	}
+
+
+	/**
+	 * Queries for all non consumed non expired purchases
+	 */
+	@Override
+	public boolean QueryExistingPurchases()
+	{
+		Log.debug("[GooglePlayStoreHelper] - QueryExistingPurchases");
+
+		QueryPurchasesParams QueryInApp =  QueryPurchasesParams.newBuilder()
+			.setProductType(BillingClient.ProductType.INAPP)
+			.build();
+		QueryPurchasesParams QuerySubs = QueryPurchasesParams.newBuilder()
+			.setProductType(BillingClient.ProductType.SUBS)
+			.build();
+
+		PurchasesResponseListener AggregatedListener = new PurchasesResponseListener() {
+			// This method is marked synchronized because although Google Billing Library documentation says "All methods annotated with AnyThread can be 
+			// called from any thread and all the asynchronous callbacks will be returned on the same thread" those callbacks may be called on different 
+			// threads if invoked from JNI 
+			@Override
+			public synchronized void onQueryPurchasesResponse(@NonNull BillingResult Result, @NonNull List<Purchase> OwnedItems)
+			{
+				if (AwaitingResponseCount <= 0)
+					return;
+
+				int ResponseCode = Result.getResponseCode();
+				if (ResponseCode == BillingClient.BillingResponseCode.OK)
+				{
+					AllProducts.addAll(OwnedItems);
+					AwaitingResponseCount -= 1;
+					Log.debug("[GooglePlayStoreHelper] - QueryExistingPurchases - Received " + OwnedItems.size() + " products. Waiting for additional responses: " + (AwaitingResponseCount != 0));
+					if (AwaitingResponseCount == 0)
+					{
+						Log.debug("[GooglePlayStoreHelper] - QueryExistingPurchases - Success! User has previously purchased " + AllProducts.size() + " inapp products");
+
+						if (AllProducts.isEmpty())
+						{
+							NativeQueryExistingPurchasesComplete(ResponseCode, new String[0][0], new int[0], new String[0], new String[0], new String[0]);
+						}
+						else
+						{
+							// Request product details so we can fix notified product ids in case they refer to a subscription
+							String[] StoreProductIds = AllProducts.stream().flatMap(Purchase -> Purchase.getProducts().stream()).toArray(String[]::new);
+							QueryProductDetailsParams QuerySubs = CreateQueryProductDetailsParams(StoreProductIds, BillingClient.ProductType.SUBS);
+							Client.queryProductDetailsAsync(QuerySubs, new ProductDetailsResponseListener() {
+								@Override
+								public void onProductDetailsResponse(@NonNull BillingResult Result, @NonNull List<ProductDetails> ProductDetailsList)
+								{
+									int ResponseCode = Result.getResponseCode();
+									if (ResponseCode == BillingClient.BillingResponseCode.OK)
+									{
+										Set<String> SubscriptionProductIds = GetSubscriptionStoreProductIds(ProductDetailsList);
+										ListOfPurchasesAsArrays PurchasesArrays = new ListOfPurchasesAsArrays(AllProducts, SubscriptionProductIds);
+										NativeQueryExistingPurchasesComplete(ResponseCode, PurchasesArrays.OwnedProducts, PurchasesArrays.PurchaseStates, PurchasesArrays.ProductTokens, PurchasesArrays.Receipts, PurchasesArrays.Signatures);
+									}
+									else
+									{
+										Log.debug("[GooglePlayStoreHelper] - QueryExistingPurchases - Failed to get product information while collecting existing purchases" + TranslateBillingResult(Result));
+										NativeQueryExistingPurchasesComplete(ResponseCode, null, null, null, null, null);
+									}
+								}
+							});
+						}
+					}
+				}
+				else
+				{
+					AwaitingResponseCount = 0;
+					Log.debug("[GooglePlayStoreHelper] - QueryExistingPurchases - Failed to collect existing purchases" + TranslateBillingResult(Result));
+					NativeQueryExistingPurchasesComplete(ResponseCode, null, null, null, null, null);
+				}
+			}
+			// We will be querying for in app products and subscriptions 
+			int AwaitingResponseCount = 2;
+			List<Purchase> AllProducts = new ArrayList<>();
+		};
+
+		Client.queryPurchasesAsync(QueryInApp, AggregatedListener);
+		Client.queryPurchasesAsync(QuerySubs, AggregatedListener);
+		return true;
+	}
+
+	/**
+	 * Locally acknowledges the purchase referred by the provided PurchaseToken
+	 */
+	@Override
+	public void AcknowledgePurchase(final String PurchaseToken)
+	{
+		Log.debug("[GooglePlayStoreHelper] - Beginning AcknowledgePurchase: " + PurchaseToken);
+
+		AcknowledgePurchaseParams Params = AcknowledgePurchaseParams.newBuilder()
+			.setPurchaseToken(PurchaseToken)
+			.build();
+		Client.acknowledgePurchase(Params, new AcknowledgePurchaseResponseListener() {
+			@Override
+			public void onAcknowledgePurchaseResponse(@NonNull BillingResult Result)
+			{
+				Log.debug("[GooglePlayStoreHelper] - AcknowledgePurchase response: " + TranslateBillingResult(Result));
+				NativeAcknowledgeComplete(Result.getResponseCode(), PurchaseToken);
+			}
+		});
+	}
+
+	/**
+	 * Locally consumes the purchase referred by the provided PurchaseToken
+	 */
+	@Override
+	public void ConsumePurchase(String PurchaseToken)
+	{
+		Log.debug("[GooglePlayStoreHelper] - Beginning ConsumePurchase: " + PurchaseToken);
+
+		ConsumeParams Params = ConsumeParams.newBuilder()
+			.setPurchaseToken(PurchaseToken)
+			.build();
+		Client.consumeAsync(Params, new ConsumeResponseListener() {
+			@Override
+			public void onConsumeResponse(@NonNull BillingResult Result, @NonNull String PurchaseToken)
+			{
+				Log.debug("[GooglePlayStoreHelper] - ConsumePurchase response: " + TranslateBillingResult(Result));
+				NativeConsumeComplete(Result.getResponseCode(), PurchaseToken);
+			}
+		});
+	}
+
+	///////////////////////////////////////////////////////
+	// The Google's PurchasesUpdatedListener interface implementation.
+
+	@Override
+	public void onPurchasesUpdated(@NonNull BillingResult Result, @Nullable final List<Purchase> UpdatedPurchases)
+	{
+		int responseCode = Result.getResponseCode();
+		Log.debug("[GooglePlayStoreHelper] - onPurchasesUpdated - Processing purchase result. Response Code: " + TranslateBillingResult(Result));
+		if (responseCode == BillingClient.BillingResponseCode.OK && UpdatedPurchases != null)
+		{
+			// Request product details so we can fix notified product ids in case they refer to a subscription
+			String[] StoreProductIds = UpdatedPurchases.stream().flatMap(Purchase -> Purchase.getProducts().stream()).toArray(String[]::new);
+			QueryProductDetailsParams QuerySubs = CreateQueryProductDetailsParams(StoreProductIds, BillingClient.ProductType.SUBS);
+			Client.queryProductDetailsAsync(QuerySubs, new ProductDetailsResponseListener() {
+				@Override
+				public void onProductDetailsResponse(@NonNull BillingResult Result, @NonNull List<ProductDetails> ProductDetailsList)
+				{
+					int responseCode = Result.getResponseCode();
+					if (responseCode == BillingClient.BillingResponseCode.OK)
+					{
+						Set<String> SubscriptionProductIds = GetSubscriptionStoreProductIds(ProductDetailsList);
+
+						for (Purchase UpdatedPurchase : UpdatedPurchases)
+						{
+							String[] ProductIds = EncodedProductId.StoreProductIdsToProductIds(UpdatedPurchase.getProducts(), SubscriptionProductIds);
+							Log.debug("[GooglePlayStoreHelper] - onPurchasesUpdated - Processing purchase of product " + String.join(",",ProductIds));
+							Log.debug("[GooglePlayStoreHelper] - Purchase data: " + UpdatedPurchase.toString());
+							Log.debug("[GooglePlayStoreHelper] - Data signature: " + UpdatedPurchase.getSignature());
+
+							String Receipt = Base64.encode(UpdatedPurchase.getOriginalJson().getBytes());
+							NativePurchaseComplete(BillingClient.BillingResponseCode.OK, ProductIds, UpdatedPurchase.getPurchaseState(), UpdatedPurchase.getPurchaseToken(), Receipt, UpdatedPurchase.getSignature());
+						}
+
+					}
+					else
+					{
+						Log.debug("[GooglePlayStoreHelper] - onPurchasesUpdated - Purchase failed while retrieving product information. Response code: " + TranslateBillingResult(Result));
+						NativePurchaseComplete(responseCode, null, Purchase.PurchaseState.UNSPECIFIED_STATE, "", "", "");
+					}
+				}
+			});
+		}
+		else
+		{
+			Log.debug("[GooglePlayStoreHelper] - onPurchasesUpdated - Purchase failed. Response code: " + TranslateBillingResult(Result));
+			NativePurchaseComplete(responseCode, null, Purchase.PurchaseState.UNSPECIFIED_STATE, "", "", "");
+		}
+	}
+
+	///////////////////////////////////////////////////////
+	// Some helpers
+
+	/**
+	 * Helper to convert Google Billing library received data to be passed to native calls
+	 */
+	class ListOfProductDetailsAsArrays
+	{
+		public String[] ProductIds;
+		public String[] Titles;
+		public String[] Descriptions;
+		public String[] Prices;
+		public long[] PricesRaw;
+		public String[] CurrencyCodes;
+
+		private int ComputeTotalEntryCount(List<ProductDetails> InProductDetailsList)
+		{
+			int TotalCount = 0;
+			for (ProductDetails ThisProduct : InProductDetailsList)
+			{
+				switch(ThisProduct.getProductType())
+				{
+					case BillingClient.ProductType.INAPP:
+						TotalCount ++;
+						break;
+					case BillingClient.ProductType.SUBS:
+						for (ProductDetails.SubscriptionOfferDetails SubsOfferDetails: ThisProduct.getSubscriptionOfferDetails())
+						{
+							TotalCount += SubsOfferDetails.getPricingPhases().getPricingPhaseList().size();
+							if (SubsOfferDetails.getOfferId() != null)
+							{
+								// En extra entry will be added for offers without including price index on it's name
+								TotalCount += 1;
+							}
+						}
+						break;
+				}
+			}
+			return TotalCount;
+		}
+
+		/**
+		 * Transforms store data into data consumable by the engine
+		 * For in app products it will create an entry with product id <store product id>
+		 * For subscription it will create an entry for each base plan with a product id <subscription prefix>:<store product id>:<base plan id>
+		 * For each offer in a subscription plan it will create an entry <subscription prefix>:<store product id>:<base plan id>:<offer id> with the 
+		 * same price data as the base plan. Then for each step in the offer it will create entries with id 
+		 * <subscription prefix>:<store product id>:<base plan id>:<offer id>:<price phase index> with increasing values for "price phase index" 
+		 */
+		public ListOfProductDetailsAsArrays(List<ProductDetails> InProductDetailsList)
+		{
+			int TotalEntryCount = ComputeTotalEntryCount(InProductDetailsList);
+			ProductIds = new String[TotalEntryCount];
+			Titles = new String[TotalEntryCount];
+			Descriptions = new String[TotalEntryCount];
+			Prices = new String[TotalEntryCount];
+			PricesRaw = new long[TotalEntryCount];
+			CurrencyCodes = new String[TotalEntryCount];
+
+			int Index = 0;
+			for (ProductDetails ThisProduct : InProductDetailsList)
+			{
+				Log.debug("[GooglePlayStoreHelper] - QueryInAppPurchases - Parsing details for: " + ThisProduct.getProductId());
+				Log.debug("[GooglePlayStoreHelper] - title: " + ThisProduct.getTitle());
+				Log.debug("[GooglePlayStoreHelper] - description: " + ThisProduct.getDescription());
+				Titles[Index] = ThisProduct.getTitle();
+				Descriptions[Index] = ThisProduct.getDescription();
+
+				switch(ThisProduct.getProductType())
+				{
+					case BillingClient.ProductType.INAPP:
+						ProductDetails.OneTimePurchaseOfferDetails PriceDetails = ThisProduct.getOneTimePurchaseOfferDetails();
+						Log.debug("[GooglePlayStoreHelper] - price: " + PriceDetails.getFormattedPrice());
+						Log.debug("[GooglePlayStoreHelper] - price_amount_micros: " + PriceDetails.getPriceAmountMicros());
+						Log.debug("[GooglePlayStoreHelper] - price_currency_code: " + PriceDetails.getPriceCurrencyCode());
+
+						ProductIds[Index] = ThisProduct.getProductId();
+						Prices[Index] = PriceDetails.getFormattedPrice();
+						PricesRaw[Index] = PriceDetails.getPriceAmountMicros();
+						CurrencyCodes[Index] = PriceDetails.getPriceCurrencyCode();
+
+						Index ++;
+						break;
+					case BillingClient.ProductType.SUBS:
+						for (ProductDetails.SubscriptionOfferDetails SubsOfferDetails: ThisProduct.getSubscriptionOfferDetails())
+						{
+							int PricingPhaseIndex = 0;
+							for( ProductDetails.PricingPhase PricingPhase: SubsOfferDetails.getPricingPhases().getPricingPhaseList())
+							{
+								Log.debug("[GooglePlayStoreHelper] - price: " + PricingPhase.getFormattedPrice());
+								Log.debug("[GooglePlayStoreHelper] - price_amount_micros: " + PricingPhase.getPriceAmountMicros());
+								Log.debug("[GooglePlayStoreHelper] - price_currency_code: " + PricingPhase.getPriceCurrencyCode());
+
+								boolean bIsOffer = SubsOfferDetails.getOfferId() != null;
+								if (bIsOffer)
+								{
+									ProductIds[Index] = EncodedProductId.EncodeSubscriptionOfferPricePoint(ThisProduct.getProductId(), SubsOfferDetails.getBasePlanId(), SubsOfferDetails.getOfferId(), PricingPhaseIndex);
+								}
+								else
+								{
+									ProductIds[Index] = EncodedProductId.EncodeSubscriptionBasePlan(ThisProduct.getProductId(), SubsOfferDetails.getBasePlanId());
+								}
+								Prices[Index] = PricingPhase.getFormattedPrice();
+								PricesRaw[Index] = PricingPhase.getPriceAmountMicros();
+								CurrencyCodes[Index] = PricingPhase.getPriceCurrencyCode();
+
+								Index++;
+								PricingPhaseIndex++;
+
+								if (bIsOffer && PricingPhaseIndex == SubsOfferDetails.getPricingPhases().getPricingPhaseList().size())
+								{
+									// Add extra product information using last price point but not including price index on it's name
+									ProductIds[Index] = EncodedProductId.EncodeSubscriptionOffer(ThisProduct.getProductId(), SubsOfferDetails.getBasePlanId(), SubsOfferDetails.getOfferId());
+									Prices[Index] = PricingPhase.getFormattedPrice();
+									PricesRaw[Index] = PricingPhase.getPriceAmountMicros();
+									CurrencyCodes[Index] = PricingPhase.getPriceCurrencyCode();
+									Index++;
+								}
+							}
+						}
+						break;
+				}
+			}
+		}
+	}
+
+	/**
+	 * Helper to convert Google Billing library received data to be passed to native calls
+	 */
+	static class ListOfPurchasesAsArrays
+	{
+		public String[][] OwnedProducts;
+		public int[] PurchaseStates;
+		public String[] ProductTokens;
+		public String[] Signatures;
+		public String[] Receipts;
+
+		public ListOfPurchasesAsArrays(List<Purchase> InOwnedItems, Set<String> SubscriptionProductId)
+		{
+			OwnedProducts = new String[InOwnedItems.size()][];
+			PurchaseStates = new int[InOwnedItems.size()];
+			ProductTokens = new String[InOwnedItems.size()];
+			Signatures = new String[InOwnedItems.size()];
+			Receipts = new String[InOwnedItems.size()];
+
+			for (int index = 0; index < InOwnedItems.size(); index ++)
+			{
+				Purchase OwnedPurchase = InOwnedItems.get(index);
+				OwnedProducts[index] = OwnedPurchase.getProducts().stream().map(ProductId -> EncodedProductId.StoreProductIdToProductId(ProductId, SubscriptionProductId)).toArray(String[]::new);
+				PurchaseStates[index] = OwnedPurchase.getPurchaseState();
+				ProductTokens[index] = OwnedPurchase.getPurchaseToken();
+				Signatures[index] = OwnedPurchase.getSignature();
+				Receipts[index] = Base64.encode(OwnedPurchase.getOriginalJson().getBytes());
+			}
+		}
+	}
+
+
+	/*
+		Helper class to Encode/Decode subscription identity data on ProductIds
+		To properly select a subscription for purchase or get it's price information ProductIds passed to/received from the engine 
+		may have a prefix and encoded identifiers to include all that information as a single ProductId
+		On the GooglePlay Store each product is identified by a single id but subscriptions may define several base plans, which define renewal periods and base pricing.
+		Each base plan may define multiple offers which include special pricings (like free trials and special pricings after that)
+		Those subscription id are in the form
+		<subscription prefix><store product id>[:<base plan id>[:<offer id>[:<offer price index>]]]
+		To query for product information just and identifier like <subscription prefix><store product id> is enough to retrieve all base plans and offers
+		In order to start a purchase the <base plan id> part is mandatory. To purchase a given offer in a base plan the <offer id> is also mandatory
+	 */
+	static class EncodedProductId
+	{
+		private final static String SUBSCRIPTION_PREFIX = "s-";
+
+		static public String EncodeSubscriptionOfferPricePoint(String StoreProductId, String BasePlanId, String OfferId, int PricingPhaseIndex)
+		{
+			return String.format("%s%s:%s:%s:%d", SUBSCRIPTION_PREFIX, StoreProductId, BasePlanId, OfferId, PricingPhaseIndex);
+		}
+		static public String EncodeSubscriptionBasePlan(String StoreProductId, String BasePlanId)
+		{
+			return String.format("%s%s:%s", SUBSCRIPTION_PREFIX, StoreProductId, BasePlanId);
+		}
+
+		static public String EncodeSubscriptionOffer(String StoreProductId, String BasePlanId, String OfferId)
+		{
+			return String.format("%s%s:%s:%s", SUBSCRIPTION_PREFIX, StoreProductId, BasePlanId, OfferId);
+		}
+
+		static public boolean IsSubscription(String InProductId)
+		{
+			return InProductId.startsWith(SUBSCRIPTION_PREFIX);
+		}
+
+		/**
+		 * returns the product type that matches all product ids. In product ids refer to mixed product types it returns null
+		 */
+		static public String GetCommonProductType(String[] InProducts)
+		{
+			if (InProducts.length == 0)
+			{
+				return null;
+			}
+
+			boolean bIsFirstProductSubscription = IsSubscription(InProducts[0]);
+
+			for (int i=1; i < InProducts.length; ++i)
+			{
+				if (bIsFirstProductSubscription != IsSubscription(InProducts[i]))
+				{
+					return null;
+				}
+			}
+			return bIsFirstProductSubscription? BillingClient.ProductType.SUBS : BillingClient.ProductType.INAPP;
+		}
+
+		/**
+		 * Fixes store product id to reflect whether it refers to a subscription or not
+		 */
+		public static String StoreProductIdToProductId(String StoreProductId, Set<String> SubscriptionStoreProductIds)
+		{
+			return SubscriptionStoreProductIds.contains(StoreProductId)? SUBSCRIPTION_PREFIX + StoreProductId : StoreProductId;
+		}
+
+		/**
+		 * Fixes a list of store product id to reflect whether they refer to a subscription or not
+		 */
+		public static String[] StoreProductIdsToProductIds(List<String> StoreProductIds, Set<String> SubscriptionStoreProductIds)
+		{
+			return StoreProductIds.stream().map(StoreProductId -> StoreProductIdToProductId(StoreProductId, SubscriptionStoreProductIds)).toArray(String[]::new );
+		}
+
+		/**
+		 * Decodes the product identity data from a ProductId 
+		 */
+		EncodedProductId(String CompositeProductId)
+		{
+			bIsSubscription = IsSubscription(CompositeProductId);
+			if (bIsSubscription)
+			{
+				String[] Parts = CompositeProductId.substring(SUBSCRIPTION_PREFIX.length()).split(":");
+				if (Parts.length > 4)
+				{
+					// Unknown format. Get it as a full product id
+					StoreProductId = CompositeProductId;
+				}
+				else
+				{
+					// The PricingPhaseIndex is not needed to start a purchase or query for products, so no need to store it
+					if (Parts.length >= 1)
+					{
+						StoreProductId = Parts[0];
+					}
+					if (Parts.length >= 2)
+					{
+						BasePlanId = Parts[1];
+					}
+					if (Parts.length >= 3)
+					{
+						OfferId = Parts[2];
+					}
+				}
+			}
+			else
+			{
+				StoreProductId = CompositeProductId;
+			}
+		}
+
+		/**
+		 * Creates BillingFlowParams.ProductDetailsParams by looking for store objects retrieved from the store that match the identity data contained on this instance 
+		 * To create a BillingFlowParams.ProductDetailsParams for a subscription we need to use the ProductDetails instance for its store product id and the offer token 
+		 * from the nested SubscriptionOfferDetails list
+		 * To create a BillingFlowParams.ProductDetailsParams for a inapp product we only need the ProductDetails matching its store product id 
+		 * Returns null if we cannot create the proper instance
+		 */
+		public BillingFlowParams.ProductDetailsParams CreateProductDetailsParams(List<ProductDetails> InProductDetails)
+		{
+			if (bIsSubscription)
+			{
+				// If a ProductDetails matches the StoreProductId and also contains an offer matching BasePlanId and OfferId then build ProductDetailsParams from it and the offer token
+				return InProductDetails.stream()
+					.filter(Details -> Details.getProductId().equals(StoreProductId) && Details.getProductType().equals(BillingClient.ProductType.SUBS))
+					.flatMap(Details -> Details.getSubscriptionOfferDetails().stream().filter(Offer -> HasMatchingOffer(Offer)).map(Offer -> new Pair<>(Details, Offer.getOfferToken())))
+					.findFirst()
+					.map(OfferData -> BillingFlowParams.ProductDetailsParams.newBuilder().setProductDetails(OfferData.first).setOfferToken(OfferData.second).build())
+					.orElse(null);
+			}
+			else
+			{
+				// If a ProductDetails matches the StoreProductId then build ProductDetailsParams from it
+				return InProductDetails.stream()
+					.filter(Details -> Details.getProductId().equals(StoreProductId) && Details.getProductType().equals(BillingClient.ProductType.INAPP))
+					.findFirst()
+					.map(Details -> BillingFlowParams.ProductDetailsParams.newBuilder().setProductDetails(Details).build())
+					.orElse(null);
+			}
+		}
+
+		/**
+		 * Check if the offer data matches the identity data contained on this instance
+		 */
+		boolean HasMatchingOffer(ProductDetails.SubscriptionOfferDetails OfferDetails)
+		{
+			if (!OfferDetails.getBasePlanId().equals(BasePlanId))
+			{
+				return false;
+			}
+
+			if (OfferId == null || OfferId.isEmpty())
+			{
+				return OfferDetails.getOfferId() == null;
+			}
+			else
+			{
+				return OfferId.equals(OfferDetails.getOfferId());
+			}
+		}
+
+		boolean bIsSubscription;
+		String StoreProductId;
+		String BasePlanId;
+		String OfferId;
+	}
+
+	/**
+	 * Creates a BillingFlowParams from product ids with encoded identity data
+	 */
+	private static BillingFlowParams CreateBillingFlowParams(String[] RequestedProductIds, List<ProductDetails> InProductDetails, String InObfuscatedAccountId)
+	{
+		List<BillingFlowParams.ProductDetailsParams> ProductDetailsParamList = Arrays.stream(RequestedProductIds)
+			.map(ProductId -> new EncodedProductId(ProductId).CreateProductDetailsParams(InProductDetails))
+			.filter(Objects::nonNull)
+			.collect(Collectors.toList());
+
+		if (ProductDetailsParamList.isEmpty())
+		{
+			return null;
+		}
+		BillingFlowParams.Builder Params = BillingFlowParams.newBuilder()
+			.setProductDetailsParamsList(ProductDetailsParamList);
+		if (InObfuscatedAccountId != null)
+		{
+			Params = Params.setObfuscatedAccountId(InObfuscatedAccountId);
+		}
+		return Params.build();
+	}
+
+	/**
+	 * Creates a QueryProductDetailsParams from product ids with encoded identity data. Only StoreProductId is needed to query information for all base plans 
+	 * and offers in a subscription prodduct
+	 */
+	private QueryProductDetailsParams CreateQueryProductDetailsParams(String[] InProductIds, String ProductType)
+	{
+		ArrayList<QueryProductDetailsParams.Product> ProductList = new ArrayList<>();
+		for (String ProductId : InProductIds)
+		{
+			EncodedProductId Info = new EncodedProductId(ProductId);
+			QueryProductDetailsParams.Product Product = QueryProductDetailsParams.Product.newBuilder()
+				.setProductId(Info.StoreProductId)
+				.setProductType(ProductType)
+				.build();
+			ProductList.add(Product);
+		}
+
+		return QueryProductDetailsParams.newBuilder()
+			.setProductList(ProductList)
+			.build();
+	}
+
+	/**
+	 * Get a text translation of the Response Codes returned by google play.
+	 */
+	private static String TranslateBillingResponseCode(int ResponseCode)
+	{
+		switch(ResponseCode)
+		{
+			case BillingClient.BillingResponseCode.OK:
+				return "BillingResponseCode.OK";
+			case BillingClient.BillingResponseCode.USER_CANCELED:
+				return "BillingResponseCode.USER_CANCELED";
+			case BillingClient.BillingResponseCode.SERVICE_UNAVAILABLE:
+				return "BillingResponseCode.SERVICE_UNAVAILABLE";
+			case BillingClient.BillingResponseCode.SERVICE_DISCONNECTED:
+				return "BillingResponseCode.SERVICE_DISCONNECTED";
+			case BillingClient.BillingResponseCode.BILLING_UNAVAILABLE:
+				return "BillingResponseCode.BILLING_UNAVAILABLE";
+			case BillingClient.BillingResponseCode.FEATURE_NOT_SUPPORTED:
+				return "BillingResponseCode.FEATURE_NOT_SUPPORTED";
+			case BillingClient.BillingResponseCode.ITEM_UNAVAILABLE:
+				return "BillingResponseCode.ITEM_UNAVAILABLE";
+			case BillingClient.BillingResponseCode.DEVELOPER_ERROR:
+				return "BillingResponseCode.DEVELOPER_ERROR";
+			case BillingClient.BillingResponseCode.ERROR:
+				return "BillingResponseCode.ERROR";
+			case BillingClient.BillingResponseCode.ITEM_ALREADY_OWNED:
+				return "BillingResponseCode.ITEM_ALREADY_OWNED";
+			case BillingClient.BillingResponseCode.ITEM_NOT_OWNED:
+				return "BillingResponseCode.ITEM_NOT_OWNED";
+			default:
+				return "Unknown Response Code: " + ResponseCode;
+		}
+	}
+
+	/**
+	 * Get an informative text with as much information about the BillingResult as possible.
+	 */
+	private static String TranslateBillingResult(@NonNull BillingResult Result)
+	{
+		int ResponseCode = Result.getResponseCode();
+		switch(ResponseCode)
+		{
+			case BillingClient.BillingResponseCode.OK:
+			case BillingClient.BillingResponseCode.USER_CANCELED:
+				return TranslateBillingResponseCode(ResponseCode);
+			default:
+				return TranslateBillingResponseCode(ResponseCode) + " DebugMessage: " + Result.getDebugMessage();
+		}
+	}
+
+	///////////////////////////////////////////////////////
+	// Game Activity/Context driven methods we need to listen for.
+
+	/**
+	 * On Destroy disconnect from service
+	 */
+	public void onDestroy()
+	{
+		Log.debug("[GooglePlayStoreHelper] - onDestroy");
+
+		if (Client != null)
+		{
+			Client.endConnection();
+		}
+	}
+
+	// Callback that notify the C++ implementation that a task has completed
+	public native void NativeQueryComplete(int ResponseCode, String[] ProductIds, String[] Titles, String[] Descriptions, String[] Prices, long[] PricesRaw, String[] CurrencyCodes);
+	public native void NativePurchaseComplete(int ResponseCode, String[] ProductIds, int PurchaseState, String ProductToken, String ReceiptData, String Signature);
+	public native void NativeQueryExistingPurchasesComplete(int ResponseCode, String[][] ProductIds, int[] PurchaseState, String[] ProductTokens, String[] ReceiptsData, String[] Signatures);
+	public native void NativeConsumeComplete(int ResponseCode, String PurchaseToken);
+	public native void NativeAcknowledgeComplete(int ResponseCode, String PurchaseToken);
+}
diff --git a/Engine/Plugins/Online/Android/OnlineSubsystemGooglePlay/Source/OnlineSubsystemGooglePlay.Build.cs b/Engine/Plugins/Online/Android/OnlineSubsystemGooglePlay/Source/OnlineSubsystemGooglePlay.Build.cs
--- a/Engine/Plugins/Online/Android/OnlineSubsystemGooglePlay/Source/OnlineSubsystemGooglePlay.Build.cs
+++ b/Engine/Plugins/Online/Android/OnlineSubsystemGooglePlay/Source/OnlineSubsystemGooglePlay.Build.cs
@@ -29,8 +29,7 @@ public class OnlineSubsystemGooglePlay : ModuleRules
 				"OnlineSubsystem", 
                 "HTTP",
 				"AndroidRuntimeSettings",
-				"Launch",
-				"GpgCppSDK"
+				"Launch"
             }
 			);
 
diff --git a/Engine/Plugins/Online/Android/OnlineSubsystemGooglePlay/Source/OnlineSubsystemGooglePlay_UPL.xml b/Engine/Plugins/Online/Android/OnlineSubsystemGooglePlay/Source/OnlineSubsystemGooglePlay_UPL.xml
--- a/Engine/Plugins/Online/Android/OnlineSubsystemGooglePlay/Source/OnlineSubsystemGooglePlay_UPL.xml
+++ b/Engine/Plugins/Online/Android/OnlineSubsystemGooglePlay/Source/OnlineSubsystemGooglePlay_UPL.xml
@@ -2,26 +2,71 @@
 <!--OnlineSubsystemGooglePlay plugin additions-->
 <root xmlns:android="http://schemas.android.com/apk/res/android">	
 	<init>
-	  <log text="Subsystem Google Play SDK Android init"/>
+		<log text="Subsystem Google Play SDK Android init"/>
+		<setBoolFromProperty result="bEnableGooglePlaySupport" ini="Engine" section="/Script/AndroidRuntimeSettings.AndroidRuntimeSettings" property="bEnableGooglePlaySupport" default="false"/>
+		<setBoolFromProperty result="bSupportsInAppPurchasing" ini="Engine" section="OnlineSubsystemGooglePlay.Store" property="bSupportsInAppPurchasing" default="false"/>
+		<setIntFromProperty result="RequestCodeForPlayGamesActivities" ini="Engine" section="/Script/AndroidRuntimeSettings.AndroidRuntimeSettings" property="RequestCodeForPlayGamesActivities" default="80002"/>
+		<dumpvars/>
 	</init>
 	
 	<prebuildCopies>
-		<log text="Copying Google Play BillingApiV2 java files"/>
-		<copyDir src="$S(PluginDir)/Java/BillingApiV2"
-				 dst = "$S(BuildDir)/src" />
+		<log text="Copying Google Play wrapper java file"/>
+		<copyFile src="$S(PluginDir)/Java/com/epicgames/unreal/GooglePlayGamesWrapper.java"
+					dst = "$S(BuildDir)/src/com/epicgames/unreal/GooglePlayGamesWrapper.java" />
+		<if condition="bSupportsInAppPurchasing">
+			<true>
+				<log text="Copying Google Play store helper java file"/>
+				<copyFile src="$S(PluginDir)/Java/com/epicgames/unreal/GooglePlayStoreHelper.java"
+						dst = "$S(BuildDir)/src/com/epicgames/unreal/GooglePlayStoreHelper.java" />
+			</true>
+		</if>
 	</prebuildCopies>
 	
 	<!-- gradle dependencies additions -->
 	<buildGradleAdditions>
-		<insert>
+		<if condition="bSupportsInAppPurchasing">
+			<true>
+				<insert>
 			dependencies {
 			implementation 'com.android.billingclient:billing:6.0.0'
 			}
-		</insert>
+				</insert>
+			</true>
+		</if>
 	</buildGradleAdditions>
 
-	<gameActivityAndroidThunkJavaIapBeginPurchase>
+	<!-- optional additions to proguard -->
+	<proguardAdditions>
+		<if condition="bEnableGooglePlaySupport">
+			<true>
+				<insert>
+ -keep class com.epicgames.unreal.GooglePlayGamesWrapper {
+	public *;
+}                   
+				</insert>
+			</true>
+		</if>
+		<if condition="bSupportsInAppPurchasing">
+			<true>
+				<insert>
+ -keep class com.epicgames.unreal.GooglePlayStoreHelper {
+	public *;
+}                   
+				</insert>
+			</true>
+		</if>
+	</proguardAdditions>
+
+	<gameApplicationOnCreateAdditions>
 		<insert>
+		GooglePlayGamesWrapper.Initialize(getApplicationContext());
+		</insert>
+	</gameApplicationOnCreateAdditions>
+
+	<gameActivityAndroidThunkJavaIapBeginPurchase>
+		<if condition="bSupportsInAppPurchasing">
+			<true>
+				<insert>
 public boolean AndroidThunkJava_IapBeginPurchase(String[] ProductIds, String AccountId)
 {
 	Log.debug("[JAVA] - AndroidThunkJava_IapBeginPurchase");
@@ -58,21 +103,33 @@ public boolean AndroidThunkJava_IapBeginPurchase(String[] ProductIds, String Acc
 	}
 	return bTriggeredPurchase;
 }
-		</insert>
+				</insert>
+			</true>
+		</if>
 	</gameActivityAndroidThunkJavaIapBeginPurchase>
 	
 	<!-- optional additions to GameActivity imports in GameActivity.java -->
 	<gameActivityImportAdditions>
-		<insert>
-import com.epicgames.unreal.GooglePlayStoreHelper;
+		<if condition="bSupportsInAppPurchasing">
+			<true>
+				<insert>
 import java.io.UnsupportedEncodingException;
 import java.security.MessageDigest;
 import java.security.NoSuchAlgorithmException;
-		</insert>
+				</insert>
+			</true>
+		</if>
 	</gameActivityImportAdditions>
 
+	<gameActivityClassAdditions>
+		<insertValue value="public static int REQUEST_CODE_SHOW_GOOGLEPLAY_UI = $I(RequestCodeForPlayGamesActivities);"/>
+		<insertNewline/>
+	</gameActivityClassAdditions>
+
 	<gameActivityIapSetupServiceAdditions>
-		<insert>
+		<if condition="bSupportsInAppPurchasing">
+			<true>
+				<insert>
 		if (getPackageManager().checkPermission("com.android.vending.BILLING", getPackageName()) == getPackageManager().PERMISSION_GRANTED)
 		{
 			IapStoreHelper = new GooglePlayStoreHelper(this, Log);
@@ -90,25 +147,17 @@ import java.security.NoSuchAlgorithmException;
 			Log.debug("[JAVA] - AndroidThunkJava_IapSetupService - You do not have the appropriate permission setup.");
 			Log.debug("[JAVA] - AndroidThunkJava_IapSetupService - Please ensure com.android.vending.BILLING is added to the manifest.");
 		}
-		</insert>
+				</insert>
+			</true>
+		</if>
 	</gameActivityIapSetupServiceAdditions>
 	
-	<!-- optional additions to GameActivity onCreate in GameActivity.java -->
-	<gameActivityOnCreateAdditions>
+	<gameActivityOnActivityResultAdditions>
 		<insert>
-			//Google Play SDK onCreate additions
-			try {
-				GoogleApiClient.Builder gbuilder = new GoogleApiClient.Builder(this);
-				gbuilder.addConnectionCallbacks(this);
-				gbuilder.addOnConnectionFailedListener(this);
-				googleClient = gbuilder.build();
-			}
-			catch (Exception e)
-			{
-				Log.debug("GoogleApiClient exception caught: " + e.toString());
-			}
-			Log.debug("googleClient is " + ((googleClient == null) ? "disabled" : "valid"));
+	if (requestCode == REQUEST_CODE_SHOW_GOOGLEPLAY_UI)
+	{
+		Log.debug("[JAVA] - GooglePlay external activity UI closed (achievements or leaderboards)");
+	}
 		</insert>
-	</gameActivityOnCreateAdditions>
-
+	</gameActivityOnActivityResultAdditions>
 </root>
diff --git a/Engine/Plugins/Online/Android/OnlineSubsystemGooglePlay/Source/Private/GooglePlayGamesWrapper.cpp b/Engine/Plugins/Online/Android/OnlineSubsystemGooglePlay/Source/Private/GooglePlayGamesWrapper.cpp
new file mode 100644
--- /dev/null
+++ b/Engine/Plugins/Online/Android/OnlineSubsystemGooglePlay/Source/Private/GooglePlayGamesWrapper.cpp
@@ -0,0 +1,390 @@
+// Copyright Epic Games, Inc. All Rights Reserved.
+
+#include "GooglePlayGamesWrapper.h"
+#include "Android/AndroidJavaEnv.h"
+#include "Android/AndroidJNI.h"
+#include "Android/AndroidApplication.h"
+#include "Containers/StringConv.h"
+#include "OnlineSubsystem.h"
+#include "OnlineAchievementsInterfaceGooglePlay.h"
+#include "OnlineAsyncTaskGooglePlayFlushLeaderboards.h"
+#include "OnlineAsyncTaskGooglePlayLogin.h"
+#include "OnlineAsyncTaskGooglePlayReadLeaderboard.h"
+#include "OnlineAsyncTaskGooglePlayQueryAchievements.h"
+#include "OnlineAsyncTaskGooglePlayWriteAchievements.h"
+
+FGooglePlayGamesWrapper::~FGooglePlayGamesWrapper()
+{
+	Reset();
+}
+
+void FGooglePlayGamesWrapper::Init()
+{
+	if (JNIEnv* Env = FAndroidApplication::GetJavaEnv())
+	{
+		PlayGamesWrapperClass = FAndroidApplication::FindJavaClassGlobalRef("com/epicgames/unreal/GooglePlayGamesWrapper");
+		UE_CLOG_ONLINE(PlayGamesWrapperClass == nullptr, Error, TEXT("Could not find com.epicgames.unreal.GooglePlayGamesWrapper class to interface with Java side"));
+
+		if (PlayGamesWrapperClass)
+		{
+			LoginMethodId = FJavaWrapper::FindStaticMethod(Env, PlayGamesWrapperClass, "Login", "(JLandroid/app/Activity;Ljava/lang/String;Z)V", false);
+			UE_CLOG_ONLINE(LoginMethodId == nullptr, Error, TEXT("Could not find \"Login\" Java method to interface with Java side"));
+
+			RequestLeaderboardScoreId = FJavaWrapper::FindStaticMethod(Env, PlayGamesWrapperClass, "RequestPlayerLeaderboardScore", "(JLandroid/app/Activity;Ljava/lang/String;)V", false);
+			UE_CLOG_ONLINE(RequestLeaderboardScoreId == nullptr, Error, TEXT("Could not find \"RequestPlayerLeaderboardScore\" Java method to interface with Java side"));
+
+			SubmitLeaderboardsScoresId = FJavaWrapper::FindStaticMethod(Env, PlayGamesWrapperClass, "SubmitLeaderboardsScores", "(JLandroid/app/Activity;[Ljava/lang/String;[J)V", false);
+			UE_CLOG_ONLINE(SubmitLeaderboardsScoresId == nullptr, Error, TEXT("Could not find \"SubmitLeaderboardsScores\" Java method to interface with Java side"));
+
+			QueryAchievementsId = FJavaWrapper::FindStaticMethod(Env, PlayGamesWrapperClass, "QueryAchievements", "(JLandroid/app/Activity;)V", false);
+			UE_CLOG_ONLINE(QueryAchievementsId == nullptr, Error, TEXT("Could not find \"QueryAchievements\" Java method to interface with Java side"));
+
+			WriteAchievementsId = FJavaWrapper::FindStaticMethod(Env, PlayGamesWrapperClass, "WriteAchievements", "(JLandroid/app/Activity;[Ljava/lang/String;[I[I)V", false);
+			UE_CLOG_ONLINE(WriteAchievementsId == nullptr, Error, TEXT("Could not find \"WriteAchievements\" Java method to interface with Java side"));
+
+			ShowAchievementsUIId = FJavaWrapper::FindStaticMethod(Env, PlayGamesWrapperClass, "ShowAchievementsUI", "(Landroid/app/Activity;)V", false);
+			UE_CLOG_ONLINE(ShowAchievementsUIId == nullptr, Error, TEXT("Could not find \"ShowAchievementsUI\" Java method to interface with Java side"));
+
+			ShowLeaderboardUIId = FJavaWrapper::FindStaticMethod(Env, PlayGamesWrapperClass, "ShowLeaderboardUI", "(Landroid/app/Activity;Ljava/lang/String;)V", false);
+			UE_CLOG_ONLINE(ShowLeaderboardUIId == nullptr, Error, TEXT("Could not find \"ShowLeaderboardUI\" Java method to interface with Java side"));
+		}
+	}
+	else
+	{
+		UE_LOG_ONLINE(Error, TEXT("Missing Java environment"));
+	}
+}
+
+void FGooglePlayGamesWrapper::Reset()
+{
+	if (PlayGamesWrapperClass)
+	{
+		if (JNIEnv* Env = FAndroidApplication::GetJavaEnv())
+		{
+			Env->DeleteGlobalRef(PlayGamesWrapperClass);
+		}
+		else
+		{
+			UE_LOG_ONLINE(Error, TEXT("Missing Java environment"));
+		}
+	}
+	PlayGamesWrapperClass = nullptr;
+	LoginMethodId = nullptr;
+	RequestLeaderboardScoreId = nullptr;
+	SubmitLeaderboardsScoresId = nullptr;
+	QueryAchievementsId = nullptr;
+	WriteAchievementsId = nullptr;
+	ShowAchievementsUIId = nullptr;
+	ShowLeaderboardUIId = nullptr;
+}
+
+bool FGooglePlayGamesWrapper::Login(FOnlineAsyncTaskGooglePlayLogin* Task, const FString& InAuthCodeClientId, bool InForceRefreshToken)
+{
+	if (PlayGamesWrapperClass != nullptr && LoginMethodId != nullptr)
+	{
+		if (JNIEnv* Env = FAndroidApplication::GetJavaEnv())
+		{
+			FScopedJavaObject<jstring> AuthCodeClientId = FJavaHelper::ToJavaString(Env, InAuthCodeClientId);
+			jboolean ForceRefreshToken = InForceRefreshToken ? JNI_TRUE : JNI_FALSE;
+
+			FJavaWrapper::CallStaticVoidMethod(Env, PlayGamesWrapperClass, LoginMethodId, reinterpret_cast<jlong>(Task), AndroidJavaEnv::GetGameActivityThis(), *AuthCodeClientId, ForceRefreshToken);
+			return true;
+		}
+	}
+	UE_LOG_ONLINE(Error, TEXT("Could not invoke Login Java method"));
+	return false;
+}
+
+bool FGooglePlayGamesWrapper::RequestLeaderboardScore(FOnlineAsyncTaskGooglePlayReadLeaderboard* Task, const FString& InLeaderboardId)
+{
+	if (PlayGamesWrapperClass != nullptr && RequestLeaderboardScoreId != nullptr)
+	{
+		if (JNIEnv* Env = FAndroidApplication::GetJavaEnv())
+		{
+			FScopedJavaObject<jstring> LeaderboardId = FJavaHelper::ToJavaString(Env, InLeaderboardId);
+
+			FJavaWrapper::CallStaticVoidMethod(Env, PlayGamesWrapperClass, RequestLeaderboardScoreId, reinterpret_cast<jlong>(Task), AndroidJavaEnv::GetGameActivityThis(), *LeaderboardId);
+			return true;
+		}
+	}
+	UE_LOG_ONLINE(Error, TEXT("Could not invoke RequestLeaderboardScore Java method"));
+	return false;
+}
+
+bool FGooglePlayGamesWrapper::FlushLeaderboardsScores(FOnlineAsyncTaskGooglePlayFlushLeaderboards* Task, const TArray<FGooglePlayLeaderboardScore>& Scores)
+{
+	if (PlayGamesWrapperClass != nullptr && SubmitLeaderboardsScoresId != nullptr)
+	{
+		if (JNIEnv* Env = FAndroidApplication::GetJavaEnv())
+		{
+			TArray<FStringView> LeaderboardIds;
+			auto UpdatedScores = NewScopedJavaObject(Env,Env->NewLongArray(Scores.Num()));
+			LeaderboardIds.Reserve(Scores.Num());
+			
+			jlong* UpdatedScoresArray = Env->GetLongArrayElements(*UpdatedScores, NULL);
+
+			for (int32 Idx = 0; Idx < Scores.Num(); ++Idx)
+			{
+				const FGooglePlayLeaderboardScore& Score = Scores[Idx]; 
+				LeaderboardIds.Add(Score.GooglePlayLeaderboardId);
+				UpdatedScoresArray[Idx] = (jlong)Score.Score;
+			}
+
+			Env->ReleaseLongArrayElements(*UpdatedScores, UpdatedScoresArray, 0);
+
+			FScopedJavaObject<jobjectArray> LeaderboardIdsArray = FJavaHelper::ToJavaStringArray(Env, LeaderboardIds);
+
+			FJavaWrapper::CallStaticVoidMethod(Env, PlayGamesWrapperClass, SubmitLeaderboardsScoresId, reinterpret_cast<jlong>(Task), AndroidJavaEnv::GetGameActivityThis(), *LeaderboardIdsArray, *UpdatedScores);
+			return true;
+		}
+	}
+	UE_LOG_ONLINE(Error, TEXT("Could not invoke SubmitLeaderboardsScores Java method"));
+	return false;
+}
+
+bool FGooglePlayGamesWrapper::ShowLeaderboardUI(const FString& InLeaderboardId)
+{
+	if (PlayGamesWrapperClass != nullptr && ShowLeaderboardUIId != nullptr)
+	{
+		if (JNIEnv* Env = FAndroidApplication::GetJavaEnv())
+		{
+			FScopedJavaObject<jstring> LeaderboardId = FJavaHelper::ToJavaString(Env, InLeaderboardId);
+
+			FJavaWrapper::CallStaticVoidMethod(Env, PlayGamesWrapperClass, ShowLeaderboardUIId, AndroidJavaEnv::GetGameActivityThis(), *LeaderboardId);
+			return true;
+		}
+	}
+	UE_LOG_ONLINE(Error, TEXT("Could not invoke ShowLeaderboardUI Java method"));
+	return false;
+}
+
+bool FGooglePlayGamesWrapper::QueryAchievements(FOnlineAsyncTaskGooglePlayQueryAchievements* Task)
+{
+	if (PlayGamesWrapperClass != nullptr && QueryAchievementsId != nullptr)
+	{
+		if (JNIEnv* Env = FAndroidApplication::GetJavaEnv())
+		{
+			FJavaWrapper::CallStaticVoidMethod(Env, PlayGamesWrapperClass, QueryAchievementsId, reinterpret_cast<jlong>(Task), AndroidJavaEnv::GetGameActivityThis());
+			return true;
+		}
+	}
+	UE_LOG_ONLINE(Error, TEXT("Could not invoke QueryAchievements Java method"));
+	return false;
+}
+
+bool FGooglePlayGamesWrapper::WriteAchievements(FOnlineAsyncTaskGooglePlayWriteAchievements* Task, const TArray<FGooglePlayAchievementWriteData>& WriteAchievementsData)
+{
+	if (PlayGamesWrapperClass != nullptr && WriteAchievementsId != nullptr)
+	{
+		if (JNIEnv* Env = FAndroidApplication::GetJavaEnv())
+		{
+			FScopedJavaObject<jobjectArray> AchievementIds = NewScopedJavaObject(Env, Env->NewObjectArray(WriteAchievementsData.Num(), FJavaWrapper::JavaStringClass, NULL));
+			FScopedJavaObject<jintArray> Actions = NewScopedJavaObject(Env, Env->NewIntArray(WriteAchievementsData.Num()));
+			FScopedJavaObject<jintArray>  Steps = NewScopedJavaObject(Env, Env->NewIntArray(WriteAchievementsData.Num()));
+
+			jint* ActionsData = Env->GetIntArrayElements(*Actions, NULL);
+			jint* StepsData = Env->GetIntArrayElements(*Steps, NULL);
+
+			for (int32 Idx = 0; Idx < WriteAchievementsData.Num(); ++Idx)
+			{
+				const FGooglePlayAchievementWriteData& Entry = WriteAchievementsData[Idx];
+
+				jstring Id = Env->NewStringUTF((const char*)StringCast<UTF8CHAR>(*Entry.GooglePlayAchievementId).Get()); 
+				Env->SetObjectArrayElement(*AchievementIds, Idx, Id);
+				Env->DeleteLocalRef(Id);
+
+				ActionsData[Idx] = static_cast<jint>(Entry.Action);
+				StepsData[Idx] = static_cast<jint>(Entry.Steps);
+			}
+
+			Env->ReleaseIntArrayElements(*Actions, ActionsData, JNI_COMMIT);
+			Env->ReleaseIntArrayElements(*Steps, StepsData, JNI_COMMIT);
+
+			FJavaWrapper::CallStaticVoidMethod(Env, PlayGamesWrapperClass, WriteAchievementsId, reinterpret_cast<jlong>(Task), AndroidJavaEnv::GetGameActivityThis(), *AchievementIds, *Actions, *Steps);
+
+			Env->ReleaseIntArrayElements(*Actions, ActionsData, JNI_ABORT);
+			Env->ReleaseIntArrayElements(*Steps, StepsData, JNI_ABORT);
+
+			return true;
+		}
+	}
+	UE_LOG_ONLINE(Error, TEXT("Could not invoke WriteAchievements Java method"));
+	return false;
+}
+
+bool FGooglePlayGamesWrapper::ShowAchievementsUI()
+{
+	if (PlayGamesWrapperClass != nullptr && ShowAchievementsUIId != nullptr)
+	{
+		if (JNIEnv* Env = FAndroidApplication::GetJavaEnv())
+		{
+			FJavaWrapper::CallStaticVoidMethod(Env, PlayGamesWrapperClass, ShowAchievementsUIId, AndroidJavaEnv::GetGameActivityThis());
+			return true;
+		}
+	}
+	UE_LOG_ONLINE(Error, TEXT("Could not invoke ShowAchievementsUI Java method"));
+	return false;
+}
+
+JNI_METHOD void Java_com_epicgames_unreal_GooglePlayGamesWrapper_nativeLoginSuccess(JNIEnv* jEnv, jclass /*Clazz*/, jlong TaskPtr, jstring InPlayerId, jstring InPlayerDisplayName, jstring InAuthCode)
+{
+	FOnlineAsyncTaskGooglePlayLogin* Task = reinterpret_cast<FOnlineAsyncTaskGooglePlayLogin*>(TaskPtr);
+
+	FString PlayerId = FJavaHelper::FStringFromParam(jEnv, InPlayerId);
+	FString PlayerDisplayName = FJavaHelper::FStringFromParam(jEnv, InPlayerDisplayName);
+	FString AuthCode = FJavaHelper::FStringFromParam(jEnv, InAuthCode);
+
+	checkf(!PlayerId.IsEmpty(), TEXT("PlayerId cannot be empty after a successful login"));
+
+	UE_LOG_ONLINE(Verbose, TEXT("nativeLoginSuccess PlayerId: %s Alias: %s Token: %s"), *PlayerId, *PlayerDisplayName , *AuthCode);
+
+	Task->SetLoginData(MoveTemp(PlayerId), MoveTemp(PlayerDisplayName), MoveTemp(AuthCode));
+	Task->bWasSuccessful = true;
+	Task->bIsComplete = true;
+}
+
+JNI_METHOD void Java_com_epicgames_unreal_GooglePlayGamesWrapper_nativeLoginFailed(JNIEnv* jEnv, jclass /*Clazz*/, jlong TaskPtr)
+{
+	FOnlineAsyncTaskGooglePlayLogin* Task = reinterpret_cast<FOnlineAsyncTaskGooglePlayLogin*>(TaskPtr);
+
+	UE_LOG_ONLINE(Verbose, TEXT("nativeLoginFailed"));
+
+	Task->bWasSuccessful = false;
+	Task->bIsComplete = true;
+}
+
+JNI_METHOD void Java_com_epicgames_unreal_GooglePlayGamesWrapper_nativeLeaderboardRequestSuccess(JNIEnv* jEnv, jclass /*Clazz*/, jlong TaskPtr, jstring InDisplayName, jstring InPlayerId, jlong InRank, jlong InRawScore)
+{
+	FOnlineAsyncTaskGooglePlayReadLeaderboard* Task = reinterpret_cast<FOnlineAsyncTaskGooglePlayReadLeaderboard*>(TaskPtr);
+
+	FString DisplayName = FJavaHelper::FStringFromParam(jEnv, InDisplayName);
+	FString PlayerId = FJavaHelper::FStringFromParam(jEnv, InPlayerId);
+	int64 Rank = (int64)InRank;
+	int64 RawScore = (int64)InRawScore;
+
+	UE_LOG_ONLINE(Verbose, TEXT("nativeLeaderboardRequestSuccess Alias: %s PlayerId: %s Rank: %lld RawScore: %lld"), *DisplayName, *PlayerId, Rank, RawScore);
+
+	Task->AddScore(DisplayName, PlayerId, Rank, RawScore);
+	Task->bWasSuccessful = true;
+	Task->bIsComplete = true;
+}
+
+JNI_METHOD void Java_com_epicgames_unreal_GooglePlayGamesWrapper_nativeLeaderboardRequestFailed(JNIEnv* jEnv, jclass /*Clazz*/, jlong TaskPtr)
+{
+	FOnlineAsyncTaskGooglePlayReadLeaderboard* Task = reinterpret_cast<FOnlineAsyncTaskGooglePlayReadLeaderboard*>(TaskPtr);
+
+	UE_LOG_ONLINE(Verbose, TEXT("nativeLeaderboardRequestFailed"));
+
+	Task->bWasSuccessful = false;
+	Task->bIsComplete = true;
+}
+
+JNI_METHOD void Java_com_epicgames_unreal_GooglePlayGamesWrapper_nativeFlushLeaderboardsCompleted(JNIEnv* jEnv, jclass /*Clazz*/, jlong TaskPtr, jboolean InSuccess)
+{
+	bool Success = (InSuccess != JNI_FALSE);
+	FOnlineAsyncTaskGooglePlayFlushLeaderboards* Task = reinterpret_cast<FOnlineAsyncTaskGooglePlayFlushLeaderboards*>(TaskPtr);
+
+	UE_LOG_ONLINE(Verbose, TEXT("nativeFlushLeaderboardsCompleted %s"), *LexToString(Success));
+
+	Task->bWasSuccessful = Success;
+	Task->bIsComplete = true;
+}
+
+
+JNI_METHOD void Java_com_epicgames_unreal_GooglePlayGamesWrapper_nativeQueryAchievementsSuccess(JNIEnv* jEnv, jclass /*Clazz*/, jlong TaskPtr, jobjectArray InAchievementIds, jintArray InTypes, jintArray InSteps, jintArray InTotalSteps, jobjectArray InTitles, jobjectArray InDescriptions, jbooleanArray InIsHidden, jlongArray InLastUpdatedTimestamps)
+{
+	FOnlineAsyncTaskGooglePlayQueryAchievements* Task = reinterpret_cast<FOnlineAsyncTaskGooglePlayQueryAchievements*>(TaskPtr);
+
+	UE_LOG_ONLINE(Verbose, TEXT("nativeQueryAchievementsSuccess"));
+
+	jsize AchievementIdsCount = jEnv->GetArrayLength(InAchievementIds);
+	jsize TypesCount = jEnv->GetArrayLength(InTypes);
+	jsize StepsCount = jEnv->GetArrayLength(InSteps);
+	jsize TotalStepsCount = jEnv->GetArrayLength(InTotalSteps);
+	jsize TitlesCount = jEnv->GetArrayLength(InTitles);
+	jsize DescriptionsCount = jEnv->GetArrayLength(InDescriptions);
+	jsize IsHiddenCount = jEnv->GetArrayLength(InIsHidden);
+	jsize LastUpdatedTimestampsCount = jEnv->GetArrayLength(InLastUpdatedTimestamps);
+
+	ensure( (AchievementIdsCount == TypesCount) &&
+			(AchievementIdsCount == StepsCount) && 
+			(AchievementIdsCount == TotalStepsCount) &&
+			(AchievementIdsCount == TitlesCount) && 
+			(AchievementIdsCount == DescriptionsCount) && 
+			(AchievementIdsCount == IsHiddenCount) &&
+			(AchievementIdsCount == LastUpdatedTimestampsCount) );
+
+	jint* Types = jEnv->GetIntArrayElements(InTypes, NULL);
+	jint* Steps = jEnv->GetIntArrayElements(InSteps, NULL);
+	jint* TotalSteps = jEnv->GetIntArrayElements(InTotalSteps, NULL);
+	jboolean* IsHidden = jEnv->GetBooleanArrayElements(InIsHidden, NULL);
+	jlong* LastUpdatedTimestamps = jEnv->GetLongArrayElements(InLastUpdatedTimestamps, NULL);
+
+	TArray<FOnlineAchievementGooglePlay> AchievementData;
+	AchievementData.Reserve(AchievementIdsCount);
+
+	TArray<FOnlineAchievementDesc> AchievementDesc;
+	AchievementDesc.Reserve(AchievementIdsCount);
+
+	for (jsize Idx = 0; Idx < AchievementIdsCount; Idx++)
+	{
+		FOnlineAchievementGooglePlay Achievement;
+		switch(Types[Idx])
+		{
+			case 0: // value for com.google.android.gms.games.achievement.Achievement.TYPE_STANDARD
+				Achievement.Type = EGooglePlayAchievementType::Standard;
+				Achievement.Progress = Steps[Idx] == 1? 100.f : 0.f;
+				break;
+			case 1: // value for com.google.android.gms.games.achievement.Achievement.TYPE_INCREMENTAL
+				Achievement.Type = EGooglePlayAchievementType::Incremental;
+				Achievement.Progress = (float)(100. * Steps[Idx] / TotalSteps[Idx]);
+				break;
+			default:
+				UE_LOG_ONLINE(Warning, TEXT("Unexpected achievement type"));
+				continue;
+		}
+		Achievement.Id = FJavaHelper::FStringFromLocalRef(jEnv, (jstring)jEnv->GetObjectArrayElement(InAchievementIds, Idx));
+		Achievement.TotalSteps = (int32)TotalSteps[Idx];
+		AchievementData.Add(MoveTemp(Achievement));
+
+		FOnlineAchievementDesc& Desc = AchievementDesc.AddDefaulted_GetRef();
+		Desc.Title = FText::FromString(FJavaHelper::FStringFromLocalRef(jEnv, (jstring)jEnv->GetObjectArrayElement(InTitles, Idx)));
+		Desc.LockedDesc = FText::FromString(FJavaHelper::FStringFromLocalRef(jEnv, (jstring)jEnv->GetObjectArrayElement(InDescriptions, Idx)));
+		Desc.UnlockedDesc = Desc.LockedDesc;
+		Desc.bIsHidden = IsHidden[Idx] != JNI_FALSE;
+		Desc.UnlockTime = FDateTime::FromUnixTimestamp(LastUpdatedTimestamps[Idx]); 
+	}
+
+	jEnv->ReleaseIntArrayElements(InTypes, Types, JNI_ABORT);
+	jEnv->ReleaseIntArrayElements(InSteps, Steps, JNI_ABORT);
+	jEnv->ReleaseIntArrayElements(InTotalSteps, TotalSteps, JNI_ABORT);
+	jEnv->ReleaseBooleanArrayElements(InIsHidden, IsHidden, JNI_ABORT);
+	jEnv->ReleaseLongArrayElements(InLastUpdatedTimestamps, LastUpdatedTimestamps, JNI_ABORT);
+
+	Task->SetAchievementsData(MoveTemp(AchievementData), MoveTemp(AchievementDesc));
+
+	Task->bWasSuccessful = true;
+	Task->bIsComplete = true;
+}
+
+
+JNI_METHOD void Java_com_epicgames_unreal_GooglePlayGamesWrapper_nativeQueryAchievementsFailed(JNIEnv* jEnv, jclass /*Clazz*/, jlong TaskPtr)
+{
+	FOnlineAsyncTaskGooglePlayQueryAchievements* Task = reinterpret_cast<FOnlineAsyncTaskGooglePlayQueryAchievements*>(TaskPtr);
+
+	UE_LOG_ONLINE(Verbose, TEXT("nativeQueryAchievementsFailed"));
+
+	Task->bWasSuccessful = false;
+	Task->bIsComplete = true;
+}
+
+JNI_METHOD void Java_com_epicgames_unreal_GooglePlayGamesWrapper_nativeWriteAchievementsCompleted(JNIEnv* jEnv, jclass /*Clazz*/, jlong TaskPtr, jobjectArray InSucceeded)
+{
+	FOnlineAsyncTaskGooglePlayWriteAchievements* Task = reinterpret_cast<FOnlineAsyncTaskGooglePlayWriteAchievements*>(TaskPtr);
+
+	TArray<FString> Succeeded = FJavaHelper::ObjectArrayToFStringTArray(jEnv, InSucceeded);
+	UE_LOG_ONLINE(Verbose, TEXT("nativeWriteAchievementsCompleted Succeeded: %s"), *FString::Join(Succeeded, TEXT(", ")));
+
+	Task->bWasSuccessful = Task->SetSucceeded(Succeeded);
+	Task->bIsComplete = true;
+}
diff --git a/Engine/Plugins/Online/Android/OnlineSubsystemGooglePlay/Source/Private/GooglePlayGamesWrapper.h b/Engine/Plugins/Online/Android/OnlineSubsystemGooglePlay/Source/Private/GooglePlayGamesWrapper.h
new file mode 100644
--- /dev/null
+++ b/Engine/Plugins/Online/Android/OnlineSubsystemGooglePlay/Source/Private/GooglePlayGamesWrapper.h
@@ -0,0 +1,48 @@
+// Copyright Epic Games, Inc. All Rights Reserved.
+
+#pragma once
+
+#include "CoreMinimal.h"
+
+THIRD_PARTY_INCLUDES_START
+#include <jni.h>
+THIRD_PARTY_INCLUDES_END
+
+#include "OnlineAchievementGooglePlayCommon.h"
+#include "OnlineLeaderboardGooglePlayCommon.h"
+
+class FOnlineAsyncTaskGooglePlayLogin;
+class FOnlineAsyncTaskGooglePlayFlushLeaderboards;
+class FOnlineAsyncTaskGooglePlayQueryAchievements;
+class FOnlineAsyncTaskGooglePlayReadLeaderboard;
+class FOnlineAsyncTaskGooglePlayWriteAchievements;
+
+class FGooglePlayGamesWrapper
+{
+public:
+	FGooglePlayGamesWrapper() = default;
+	FGooglePlayGamesWrapper(const FGooglePlayGamesWrapper& Other) = delete;
+	FGooglePlayGamesWrapper& operator=(const FGooglePlayGamesWrapper& Other) = delete;
+	~FGooglePlayGamesWrapper();
+
+	void Init();
+	void Reset();
+
+	bool Login(FOnlineAsyncTaskGooglePlayLogin* Task, const FString& InAuthCodeClientId, bool InForceRefreshToken);
+	bool RequestLeaderboardScore(FOnlineAsyncTaskGooglePlayReadLeaderboard* Task, const FString& LeaderboardId);
+	bool FlushLeaderboardsScores(FOnlineAsyncTaskGooglePlayFlushLeaderboards* Task, const TArray<FGooglePlayLeaderboardScore>& Scores);
+	bool QueryAchievements(FOnlineAsyncTaskGooglePlayQueryAchievements* Task);
+	bool WriteAchievements(FOnlineAsyncTaskGooglePlayWriteAchievements* Task, const TArray<FGooglePlayAchievementWriteData>& WriteAchievementsData);
+	bool ShowAchievementsUI();
+	bool ShowLeaderboardUI(const FString& LeaderboardName);
+
+private:
+	jclass PlayGamesWrapperClass = nullptr;
+	jmethodID LoginMethodId = nullptr;
+	jmethodID RequestLeaderboardScoreId = nullptr;
+	jmethodID SubmitLeaderboardsScoresId = nullptr;
+	jmethodID QueryAchievementsId = nullptr;
+	jmethodID WriteAchievementsId = nullptr;
+	jmethodID ShowAchievementsUIId = nullptr;
+	jmethodID ShowLeaderboardUIId = nullptr;
+};
diff --git a/Engine/Plugins/Online/Android/OnlineSubsystemGooglePlay/Source/Private/OnlineAchievementGooglePlayCommon.h b/Engine/Plugins/Online/Android/OnlineSubsystemGooglePlay/Source/Private/OnlineAchievementGooglePlayCommon.h
new file mode 100644
--- /dev/null
+++ b/Engine/Plugins/Online/Android/OnlineSubsystemGooglePlay/Source/Private/OnlineAchievementGooglePlayCommon.h
@@ -0,0 +1,35 @@
+// Copyright Epic Games, Inc. All Rights Reserved.
+
+#pragma once
+
+#include "Interfaces/OnlineAchievementsInterface.h"
+
+enum class FGooglePlayAchievementWriteAction : int32
+{
+	Unlock = 0, // Unlock the achievement
+	WriteSteps = 1 // Update achievement steps accomplished
+};
+
+struct FGooglePlayAchievementWriteData
+{
+	// Achievement Id as shown in GooglePlay Console
+	FString GooglePlayAchievementId;
+	// Action to perform
+	FGooglePlayAchievementWriteAction Action = FGooglePlayAchievementWriteAction::Unlock;
+	// Value to set in cas ethe action requires it
+	int32 Steps = 0;
+};
+
+enum class EGooglePlayAchievementType
+{
+	Incremental, // Allows setting a number of steps until completion
+	Standard // Just allows unlocking
+};
+
+struct FOnlineAchievementGooglePlay: FOnlineAchievement
+{
+	// Achievement type
+	EGooglePlayAchievementType Type = EGooglePlayAchievementType::Standard;
+	// Total number of steps needed to complete in case Type is EGooglePlayAchievementType::Incremental
+	int32 TotalSteps = 0;
+};
diff --git a/Engine/Plugins/Online/Android/OnlineSubsystemGooglePlay/Source/Private/OnlineAchievementsInterfaceGooglePlay.cpp b/Engine/Plugins/Online/Android/OnlineSubsystemGooglePlay/Source/Private/OnlineAchievementsInterfaceGooglePlay.cpp
--- a/Engine/Plugins/Online/Android/OnlineSubsystemGooglePlay/Source/Private/OnlineAchievementsInterfaceGooglePlay.cpp
+++ b/Engine/Plugins/Online/Android/OnlineSubsystemGooglePlay/Source/Private/OnlineAchievementsInterfaceGooglePlay.cpp
@@ -1,65 +1,110 @@
 // Copyright Epic Games, Inc. All Rights Reserved.
 
 #include "OnlineAchievementsInterfaceGooglePlay.h"
-#include "OnlineAsyncTaskManagerGooglePlay.h"
-#include "Android/AndroidJNI.h"
+#include "AndroidRuntimeSettings.h"
 #include "OnlineAsyncTaskGooglePlayQueryAchievements.h"
-#include "OnlineSubsystemGooglePlay.h"
-#include "UObject/Class.h"
+#include "OnlineAsyncTaskGooglePlayWriteAchievements.h"
 #include "OnlineIdentityInterfaceGooglePlay.h"
-
-using namespace gpg;
+#include "OnlineSubsystemGooglePlay.h"
 
 FOnlineAchievementsGooglePlay::FOnlineAchievementsGooglePlay( FOnlineSubsystemGooglePlay* InSubsystem )
-	: AndroidSubsystem(InSubsystem)
+	: Subsystem(InSubsystem)
 {
 	// Make sure the Google achievement cache is initialized to an invalid status so we know to query it first
-	GoogleAchievements.status = ResponseStatus::ERROR_TIMEOUT;
+}
+
+TOptional<FString> FOnlineAchievementsGooglePlay::GetUnrealAchievementIdFromGoogleAchievementId(const UAndroidRuntimeSettings* Settings, const FString& GoogleId)
+{
+	for(const auto& Mapping : Settings->AchievementMap)
+	{
+		if(Mapping.AchievementID == GoogleId)
+		{
+			return Mapping.Name;
+		}
+	}
+	return NullOpt;
+}
+
+TOptional<FString> FOnlineAchievementsGooglePlay::GetGoogleAchievementIdFromUnrealAchievementId(const UAndroidRuntimeSettings* Settings, const FString& UnrealId)
+{
+	for(const auto& Mapping : Settings->AchievementMap)
+	{
+		if(Mapping.Name == UnrealId)
+		{
+			return Mapping.AchievementID;
+		}
+	}
+	return NullOpt;
 }
 
 void FOnlineAchievementsGooglePlay::QueryAchievements(const FUniqueNetId& PlayerId, const FOnQueryAchievementsCompleteDelegate& Delegate)
 {
-	if (AndroidSubsystem->GetGameServices() == nullptr)
+	if (!IsLocalPlayer(PlayerId))
 	{
+		UE_LOG_ONLINE_ACHIEVEMENTS(Warning, TEXT("QueryAchievements failed because was called using non local player or local player is not logged in"));
 		Delegate.ExecuteIfBound(PlayerId, false);
 		return;
 	}
 
-	auto QueryTask = new FOnlineAsyncTaskGooglePlayQueryAchievements(AndroidSubsystem, FUniqueNetIdGooglePlay::Cast(PlayerId), Delegate);
-	AndroidSubsystem->QueueAsyncTask(QueryTask);
+	Subsystem->QueueAsyncTask(new FOnlineAsyncTaskGooglePlayQueryAchievements(Subsystem, PlayerId.AsShared(), Delegate));
 }
 
+bool FOnlineAchievementsGooglePlay::IsLocalPlayer(const FUniqueNetId& PlayerId) const
+{
+	if (!PlayerId.IsValid())
+	{
+		return false;
+	}
+	FOnlineIdentityGooglePlayPtr IdentityInt = Subsystem->GetIdentityGooglePlay();
+	FUniqueNetIdPtr LocalPlayerNetId = IdentityInt->GetUniquePlayerId(0);
+	return LocalPlayerNetId? *LocalPlayerNetId == PlayerId : false;
+}
 
 void FOnlineAchievementsGooglePlay::QueryAchievementDescriptions(const FUniqueNetId& PlayerId, const FOnQueryAchievementsCompleteDelegate& Delegate)
 {
 	// Just query achievements to get descriptions
-	// FIXME: This feels a little redundant, but we can see how platforms evolve, and make a decision then
 	QueryAchievements( PlayerId, Delegate );
 }
 
 void FOnlineAchievementsGooglePlay::FinishAchievementWrite(
-	const FUniqueNetId& PlayerId,
+	const FUniqueNetId& PlayerId, 
 	const bool bWasSuccessful,
 	FOnlineAchievementsWriteRef WriteObject,
 	FOnAchievementsWrittenDelegate Delegate)
 {
-	// If the Google achievement cache is invalid, we can't do anything here - need to
+	// If the achievements cache is invalid, we can't do anything here - need to
 	// be able to get achievement type and max steps for incremental achievements.
-	if (GoogleAchievements.status != ResponseStatus::VALID && GoogleAchievements.status != ResponseStatus::VALID_BUT_STALE)
+	if (Achievements.IsEmpty())
 	{
 		Delegate.ExecuteIfBound(PlayerId, false);
 		return;
 	}
 
-	auto DefaultSettings = GetDefault<UAndroidRuntimeSettings>();
+	TArray<FGooglePlayAchievementWriteData> AchievementData;
 
-	for (FStatPropertyArray::TConstIterator It(WriteObject->Properties); It; ++It)
+	auto Settings = GetDefault<UAndroidRuntimeSettings>();
+    bool bAllSucceeded = true;
+    
+	for (auto& [Key, Stat] : WriteObject->Properties)
 	{
-		// Access the stat and the value.
-		const FVariantData& Stat = It.Value();
-
 		// Create an achievement object which should be reported to the server.
-		const FString UnrealAchievementId(It.Key().ToString());
+		const FOnlineAchievementGooglePlay* Achievement = Achievements.FindByPredicate([Key = Key.ToString()](const FOnlineAchievement& Achievement) { return Achievement.Id == Key;} );
+		
+		if (!Achievement)
+		{
+            UE_LOG_ONLINE_ACHIEVEMENTS(Error, TEXT("Unknown achievement id: %s"), *Key.ToString());
+            bAllSucceeded = false;
+			continue;
+		}
+
+		TOptional<FString> GoogleAchievementId = GetGoogleAchievementIdFromUnrealAchievementId(Settings, Achievement->Id);
+		if (!GoogleAchievementId)
+		{
+            UE_LOG_ONLINE_ACHIEVEMENTS(Error, TEXT("No Google achievement Id for achievement id: %s"), *Key.ToString());
+            bAllSucceeded = false;
+			continue;
+		}
+
 		float PercentComplete = 0.0f;
 
 		// Setup the percentage complete with the value we are writing from the variant type
@@ -81,34 +126,27 @@ void FOnlineAchievementsGooglePlay::FinishAchievementWrite(
 			}
 		default:
 			{
+                bAllSucceeded = true;
 				UE_LOG_ONLINE_ACHIEVEMENTS(Error, TEXT("FOnlineAchievementsGooglePlay Trying to write an achievement with incompatible format. Not a float or int"));
-				break;
+				continue;
 			}
 		}
 
-		const auto GoogleAchievement = GetGoogleAchievementFromUnrealId(UnrealAchievementId);
-		
-		if (!GoogleAchievement.Valid())
-		{
-			continue;
-		}
 
-		const FString GoogleId(GoogleAchievement.Id().c_str());
-
-		UE_LOG_ONLINE_ACHIEVEMENTS(Log, TEXT("Writing achievement name: %s, Google id: %s, progress: %.0f"),
-			*UnrealAchievementId, *GoogleId, PercentComplete);
-		
-		switch(GoogleAchievement.Type())
+		switch(Achievement->Type)
 		{
-			case gpg::AchievementType::INCREMENTAL:
+			case EGooglePlayAchievementType::Incremental:
 			{
-				float StepFraction = (PercentComplete / 100.0f) * GoogleAchievement.TotalSteps();
+				float StepFraction = (PercentComplete / 100.0f) * Achievement->TotalSteps;
 				int RoundedSteps =  FMath::RoundToInt(StepFraction);
 
 				if(RoundedSteps > 0)
 				{
 					UE_LOG_ONLINE_ACHIEVEMENTS(Log, TEXT("  Incremental: setting progress to %d"), RoundedSteps);
-					AndroidSubsystem->GetGameServices()->Achievements().SetStepsAtLeast(GoogleAchievement.Id(), RoundedSteps);
+					FGooglePlayAchievementWriteData& Entry = AchievementData.Emplace_GetRef();
+					Entry.GooglePlayAchievementId = *GoogleAchievementId;
+					Entry.Action = FGooglePlayAchievementWriteAction::WriteSteps;
+					Entry.Steps = RoundedSteps;
 				}
 				else
 				{
@@ -117,44 +155,51 @@ void FOnlineAchievementsGooglePlay::FinishAchievementWrite(
 				break;
 			}
 
-			case gpg::AchievementType::STANDARD:
+			case EGooglePlayAchievementType::Standard:
 			{
 				// Standard achievements only unlock if the progress is at least 100%.
 				if (PercentComplete >= 100.0f)
 				{
 					UE_LOG_ONLINE_ACHIEVEMENTS(Log, TEXT("  Standard: unlocking"));
-					AndroidSubsystem->GetGameServices()->Achievements().Unlock(GoogleAchievement.Id());
+					FGooglePlayAchievementWriteData& Entry = AchievementData.Emplace_GetRef();
+					Entry.GooglePlayAchievementId = *GoogleAchievementId;
+					Entry.Action = FGooglePlayAchievementWriteAction::Unlock;
 				}
 				break;
 			}
 		}
 	}
 
-	Delegate.ExecuteIfBound(PlayerId, true);
+	if (AchievementData.IsEmpty())
+	{
+		Delegate.ExecuteIfBound(PlayerId, bAllSucceeded);
+	}
+	else
+	{
+		auto QueryTask = new FOnlineAsyncTaskGooglePlayWriteAchievements(Subsystem, PlayerId.AsShared(), MoveTemp(AchievementData), Delegate);
+		Subsystem->QueueAsyncTask(QueryTask);
+	}
 }
 
 void FOnlineAchievementsGooglePlay::WriteAchievements( const FUniqueNetId& PlayerId, FOnlineAchievementsWriteRef& WriteObject, const FOnAchievementsWrittenDelegate& Delegate )
 {
-	// @todo: verify the PlayerId passed in is the currently signed in player
-		
-	if (AndroidSubsystem->GetGameServices() == nullptr)
+	if (!IsLocalPlayer(PlayerId))
 	{
+		UE_LOG_ONLINE_ACHIEVEMENTS(Warning, TEXT("WriteAchievements failed because was called using non local player or local player is not logged in"));
 		Delegate.ExecuteIfBound(PlayerId, false);
 		return;
 	}
-
-	// We need valid data from Google first because we use the achievement type to convert Unreal progress
-	// (percentage) to Google progress (binary locked/unlocked or integer number of arbitrary steps).
-	// Kick off a query if we don't have valid data.
-	if (GoogleAchievements.status != ResponseStatus::VALID)
+		
+	if (Achievements.IsEmpty())
 	{
-		auto QueryTask = new FOnlineAsyncTaskGooglePlayQueryAchievements(AndroidSubsystem, FUniqueNetIdGooglePlay::Cast(PlayerId),
+		auto QueryTask = new FOnlineAsyncTaskGooglePlayQueryAchievements(Subsystem,
+			PlayerId.AsShared(),
 			FOnQueryAchievementsCompleteDelegate::CreateRaw(
 				this,
 				&FOnlineAchievementsGooglePlay::FinishAchievementWrite,
 				WriteObject,
 				Delegate));
-		AndroidSubsystem->QueueAsyncTask(QueryTask);
+		Subsystem->QueueAsyncTask(QueryTask);
 	}
 	else
 	{
@@ -164,27 +209,22 @@ void FOnlineAchievementsGooglePlay::WriteAchievements( const FUniqueNetId& Playe
 
 EOnlineCachedResult::Type FOnlineAchievementsGooglePlay::GetCachedAchievements(const FUniqueNetId& PlayerId, TArray<FOnlineAchievement>& OutAchievements)
 {
-	// @todo: verify the PlayerId passed in is the currently signed in player
+	if (!IsLocalPlayer(PlayerId))
+	{
+		UE_LOG_ONLINE_ACHIEVEMENTS(Warning, TEXT("GetCachedAchievements failed because was called using non local player or local player is not logged in"));
+		return EOnlineCachedResult::NotFound;
+	}
 
 	OutAchievements.Empty();
 
-	if (GoogleAchievements.status != ResponseStatus::VALID)
+	if (Achievements.IsEmpty())
 	{
 		return EOnlineCachedResult::NotFound;
 	}
 
-	auto DefaultSettings = GetDefault<UAndroidRuntimeSettings>();
-
-	for (auto& CurrentAchievement : GoogleAchievements.data)
+	for (auto& CurrentAchievement : Achievements)
 	{
-		const auto UnrealAchievement = GetUnrealAchievementFromGoogleAchievement(DefaultSettings, CurrentAchievement);
-		if (UnrealAchievement.Id.IsEmpty())
-		{
-			// Empty id means the achievement wasn't found.
-			continue;
-		}
-
-		OutAchievements.Add(UnrealAchievement);
+		OutAchievements.Add(CurrentAchievement);
 	}
 
 	return EOnlineCachedResult::Success;
@@ -193,17 +233,13 @@ EOnlineCachedResult::Type FOnlineAchievementsGooglePlay::GetCachedAchievements(c
 
 EOnlineCachedResult::Type FOnlineAchievementsGooglePlay::GetCachedAchievementDescription(const FString& AchievementId, FOnlineAchievementDesc& OutAchievementDesc)
 {
-	auto FoundAchievement = GetGoogleAchievementFromUnrealId(AchievementId);
-	if (!FoundAchievement.Valid())
+	FOnlineAchievementDesc* AchievementDescription = AchievementDescriptions.Find(AchievementId);
+	if (!AchievementDescription)
 	{
 		return EOnlineCachedResult::NotFound;
 	}
 
-	OutAchievementDesc.Title = FText::FromString(FString(FoundAchievement.Name().c_str()));
-	OutAchievementDesc.LockedDesc = FText::FromString(FString(FoundAchievement.Description().c_str()));
-	OutAchievementDesc.UnlockedDesc = FText::FromString(FString(FoundAchievement.Description().c_str()));
-	OutAchievementDesc.bIsHidden = FoundAchievement.State() == AchievementState::HIDDEN;
-	OutAchievementDesc.UnlockTime = FDateTime(0); // @todo: this
+	OutAchievementDesc = *AchievementDescription;
 
 	return EOnlineCachedResult::Success;
 }
@@ -211,124 +247,95 @@ EOnlineCachedResult::Type FOnlineAchievementsGooglePlay::GetCachedAchievementDes
 #if !UE_BUILD_SHIPPING
 bool FOnlineAchievementsGooglePlay::ResetAchievements( const FUniqueNetId& PlayerId )
 {
-	UE_LOG_ONLINE_ACHIEVEMENTS(Log, TEXT("Resetting Google Play achievements."));
-
-	extern void AndroidThunkCpp_ResetAchievements();
-	AndroidThunkCpp_ResetAchievements();
-
+	// Not supported client side. Use Management APIs for this: https://developers.google.com/games/services/management/api/achievements/reset
 	return false;
 };
 #endif // !UE_BUILD_SHIPPING
 
 EOnlineCachedResult::Type FOnlineAchievementsGooglePlay::GetCachedAchievement( const FUniqueNetId& PlayerId, const FString& AchievementId, FOnlineAchievement& OutAchievement )
 {
-	// @todo: verify the PlayerId passed in is the currently signed in player
+	if (!IsLocalPlayer(PlayerId))
+	{
+		UE_LOG_ONLINE_ACHIEVEMENTS(Warning, TEXT("GetCachedAchievement failed because was called using non local player or local player is not logged in"));
+		return EOnlineCachedResult::NotFound;
+	}
 
-	auto FoundAchievement = GetGoogleAchievementFromUnrealId(AchievementId);
-	if (!FoundAchievement.Valid())
+	FOnlineAchievement* FoundAchievement = Achievements.FindByPredicate([AchievementId](const FOnlineAchievement& Achievement) { return Achievement.Id == AchievementId; });
+	if (!FoundAchievement)
 	{
 		return EOnlineCachedResult::NotFound;
 	}
 
-	OutAchievement.Id = AchievementId;
-	OutAchievement.Progress = GetProgressFromGoogleAchievement(FoundAchievement);
+	OutAchievement = *FoundAchievement;
 
 	return EOnlineCachedResult::Success;
 }
 
-Achievement FOnlineAchievementsGooglePlay::GetGoogleAchievementFromUnrealId(const FString& UnrealId) const
+void FOnlineAchievementsGooglePlay::UpdateCache(TArray<FOnlineAchievementGooglePlay>&& Results, TArray<FOnlineAchievementDesc>&& Descriptions)
 {
-	auto DefaultSettings = GetDefault<UAndroidRuntimeSettings>();
+	ensure(Results.Num() == Descriptions.Num());
 
-	FString TargetGoogleId;
+	int32 Count = FPlatformMath::Min(Results.Num(), Descriptions.Num());
+	AchievementDescriptions.Reset();
+	Achievements.Reset(Count);
 
-	for(const auto& Mapping : DefaultSettings->AchievementMap)
+	auto Settings = GetDefault<UAndroidRuntimeSettings>();
+
+	for (int32 Idx = 0; Idx < Count; ++Idx)
 	{
-		if(Mapping.Name.Equals(UnrealId))
-		{
-			TargetGoogleId = Mapping.AchievementID;
-			break;
+		FOnlineAchievementGooglePlay& Achievement = Results[Idx];
+		if (TOptional<FString> Id = GetUnrealAchievementIdFromGoogleAchievementId(Settings, Achievement.Id))
+		{ 
+			Achievement.Id = MoveTemp(*Id);
+			AchievementDescriptions.Emplace(Achievement.Id, MoveTemp(Descriptions[Idx]));
+			Achievements.Add(MoveTemp(Achievement));
 		}
 	}
+}
 
-	for (auto& CurrentAchievement : GoogleAchievements.data)
+void FOnlineAchievementsGooglePlay::UpdateCacheAfterWrite(const TArray<FGooglePlayAchievementWriteData>& WrittenData)
+{
+	auto Settings = GetDefault<UAndroidRuntimeSettings>();
+
+	for (const FGooglePlayAchievementWriteData& Data: WrittenData)
 	{
-		if (!CurrentAchievement.Valid())
+		FOnlineAchievementGooglePlay* Achievement = nullptr;
+
+		if (TOptional<FString> UnrealAchievementId = GetUnrealAchievementIdFromGoogleAchievementId(Settings, Data.GooglePlayAchievementId))
 		{
-			continue;
+			Achievement = Achievements.FindByPredicate([&UnrealAchievementId](const FOnlineAchievementGooglePlay& Achievement)
+				{
+					return Achievement.Id == *UnrealAchievementId;
+				});
 		}
 
-		const FString CurrentGoogleId = CurrentAchievement.Id().c_str();
-
-		if (CurrentGoogleId.Equals(TargetGoogleId))
+		if (!Achievement)
 		{
-			return CurrentAchievement;
+            UE_LOG_ONLINE_ACHIEVEMENTS(Error, TEXT("Unknown Google achievement id: %s"), *Data.GooglePlayAchievementId);
+			continue;
 		}
-	}
 
-	// Not found, return empty achievement, Valid() will return false for this one.
-	return Achievement();
-}
-
-FString FOnlineAchievementsGooglePlay::GetUnrealIdFromGoogleId(const UAndroidRuntimeSettings* Settings, const FString& GoogleId)
-{
-	for(const auto& Mapping : Settings->AchievementMap)
-	{
-		if(Mapping.AchievementID.Equals(GoogleId))
+		switch (Data.Action)
 		{
-			return Mapping.Name;
+			case FGooglePlayAchievementWriteAction::WriteSteps:
+			{
+				float NewProgress = (float)( (100. * Data.Steps) / Achievement->TotalSteps);
+				if (NewProgress > Achievement->Progress)
+				{ 
+					Achievement->Progress = NewProgress;
+				}
+				break;
+			}
+			case FGooglePlayAchievementWriteAction::Unlock:
+			{
+				Achievement->Progress = 100.f;
+				break;
+			}
 		}
 	}
-
-	return FString();
-}
-
-void FOnlineAchievementsGooglePlay::UpdateCache(const AchievementManager::FetchAllResponse& Results)
-{
-	GoogleAchievements = Results;
 }
 
 void FOnlineAchievementsGooglePlay::ClearCache()
 {
-	GoogleAchievements.status = ResponseStatus::ERROR_TIMEOUT; // Is there a better error to use here?
-	GoogleAchievements.data.clear();
-}
-
-double FOnlineAchievementsGooglePlay::GetProgressFromGoogleAchievement(const Achievement& InAchievement)
-{
-	if (!InAchievement.Valid())
-	{
-		return 0.0f;
-	}
-
-	if (InAchievement.State() == AchievementState::UNLOCKED)
-	{
-		return 100.0;
-	}
-
-	if (InAchievement.Type() == AchievementType::INCREMENTAL)
-	{
-		double Fraction = (double)InAchievement.CurrentSteps() / (double)InAchievement.TotalSteps();
-		return Fraction * 100.0;
-	}
-	
-	return 0.0;
-}
-
-FOnlineAchievement FOnlineAchievementsGooglePlay::GetUnrealAchievementFromGoogleAchievement(const UAndroidRuntimeSettings* Settings, const gpg::Achievement& GoogleAchievement)
-{
-	FOnlineAchievement OutAchievement;
-
-	if (!GoogleAchievement.Valid())
-	{
-		OutAchievement.Progress = 0.0;
-		return OutAchievement;
-	}
-
-	const FString GoogleId = GoogleAchievement.Id().c_str();
-
-	OutAchievement.Id = GetUnrealIdFromGoogleId(Settings, GoogleId);
-	OutAchievement.Progress = GetProgressFromGoogleAchievement(GoogleAchievement);
-
-	return OutAchievement;
+	Achievements.Empty();
 }
diff --git a/Engine/Plugins/Online/Android/OnlineSubsystemGooglePlay/Source/Private/OnlineAchievementsInterfaceGooglePlay.h b/Engine/Plugins/Online/Android/OnlineSubsystemGooglePlay/Source/Private/OnlineAchievementsInterfaceGooglePlay.h
--- a/Engine/Plugins/Online/Android/OnlineSubsystemGooglePlay/Source/Private/OnlineAchievementsInterfaceGooglePlay.h
+++ b/Engine/Plugins/Online/Android/OnlineSubsystemGooglePlay/Source/Private/OnlineAchievementsInterfaceGooglePlay.h
@@ -3,16 +3,10 @@
 #pragma once
 
 #include "CoreMinimal.h"
-#include "Interfaces/OnlineAchievementsInterface.h"
 #include "OnlineSubsystemTypes.h"
-#include "AndroidRuntimeSettings.h"
+#include "OnlineAchievementGooglePlayCommon.h"
 #include "OnlineSubsystemGooglePlayPackage.h"
 
-THIRD_PARTY_INCLUDES_START
-#include <jni.h>
-#include "gpg/achievement_manager.h"
-THIRD_PARTY_INCLUDES_END
-
 /**
  *	IOnlineAchievements - Interface class for Achievements
  */
@@ -20,37 +14,20 @@ class FOnlineAchievementsGooglePlay : public IOnlineAchievements
 {
 private:
 
-	/** Reference to the main GameCenter subsystem */
-	class FOnlineSubsystemGooglePlay* AndroidSubsystem;
+	/** Reference to the main subsystem */
+	class FOnlineSubsystemGooglePlay* Subsystem;
 	
-	/** hide the default constructor, we need a reference to our OSS */
-	FOnlineAchievementsGooglePlay() {}
-
-	/** Our own cache of achievement data directly from Google Play */
-	gpg::AchievementManager::FetchAllResponse GoogleAchievements;
-
-	/** Looks up the Google achievement id in the mapping and returns the corresponding cached Achievement object */
-	gpg::Achievement GetGoogleAchievementFromUnrealId(const FString& UnrealId) const;
-
-	/** Uses the mapping to convert a Google achievement id to an Unreal achievement name */
-	static FString GetUnrealIdFromGoogleId(const UAndroidRuntimeSettings* Settings, const FString& GoogleId);
+	/** Cached achievements retrieved from Google Play */
+	TArray<FOnlineAchievementGooglePlay> Achievements;
 
-	/** Convert the progress of a Google achievement to a double percentage 0.0 - 100.0 */
-	static double GetProgressFromGoogleAchievement(const gpg::Achievement& InAchievement);
+	/** Cached achievement descriptions retrieved from Google Play */
+	TMap<FString, FOnlineAchievementDesc> AchievementDescriptions;
 
-	/**
-	 * Convenience function to create an Unreal achievement from a Google achievement
-	 *
-	 * @param Settings the Android runtime settings containing the achievement id mapping
-	 * @param GoogleAchievement the Google achievement object to convert
-	 */
-	static FOnlineAchievement GetUnrealAchievementFromGoogleAchievement(
-		const UAndroidRuntimeSettings* Settings,
-		const gpg::Achievement& GoogleAchievement);
+	static TOptional<FString> GetUnrealAchievementIdFromGoogleAchievementId(const class UAndroidRuntimeSettings* Settings, const FString& GoogleId);
+	static TOptional<FString> GetGoogleAchievementIdFromUnrealAchievementId(const class UAndroidRuntimeSettings* Settings, const FString& UnrealId);
 
 	/**
-	 * Using the WriteObject, fires off achievement progress calls to the Google backend. Non-blocking.
-	 * The GoogleAchievements list should be valid before this is called.
+	 * Using the WriteObject, fires off achievement progress calls to the Google Java backend. Non-blocking.
 	 *
 	 * @param PlayerId the id of the player who's making progress
 	 * @param bWasSuccessful whether a previous QueryAchievements call was successful
@@ -63,13 +40,17 @@ private:
 		FOnlineAchievementsWriteRef WriteObject,
 		FOnAchievementsWrittenDelegate Delegate);
 	
+	/** Asks the identity interface it this PlayerId refers to the local player. We can only get information about the local player from GooglePlay */
+	bool IsLocalPlayer(const FUniqueNetId& PlayerId) const;
 PACKAGE_SCOPE:
 	/** Clears the cache of Google achievements that was populated by a QueryAchievements() call. */
 	void ClearCache();
 
 	/** Called from the query achievements task to fill in the cache. */
-	void UpdateCache(const gpg::AchievementManager::FetchAllResponse& Results);
+	void UpdateCache(TArray<FOnlineAchievementGooglePlay>&& Results, TArray<FOnlineAchievementDesc>&& Descriptions);
 
+	/** Called from the write achievements task to update cache with written data. */
+	void UpdateCacheAfterWrite(const TArray<FGooglePlayAchievementWriteData>& WrittenData);
 public:
 
 	//~ Begin IOnlineAchievements Interface
diff --git a/Engine/Plugins/Online/Android/OnlineSubsystemGooglePlay/Source/Private/OnlineAsyncTaskGooglePlayAuthAction.cpp b/Engine/Plugins/Online/Android/OnlineSubsystemGooglePlay/Source/Private/OnlineAsyncTaskGooglePlayAuthAction.cpp
deleted file mode 100644
diff --git a/Engine/Plugins/Online/Android/OnlineSubsystemGooglePlay/Source/Private/OnlineAsyncTaskGooglePlayAuthAction.h b/Engine/Plugins/Online/Android/OnlineSubsystemGooglePlay/Source/Private/OnlineAsyncTaskGooglePlayAuthAction.h
deleted file mode 100644
diff --git a/Engine/Plugins/Online/Android/OnlineSubsystemGooglePlay/Source/Private/OnlineAsyncTaskGooglePlayFlushLeaderboards.cpp b/Engine/Plugins/Online/Android/OnlineSubsystemGooglePlay/Source/Private/OnlineAsyncTaskGooglePlayFlushLeaderboards.cpp
new file mode 100644
--- /dev/null
+++ b/Engine/Plugins/Online/Android/OnlineSubsystemGooglePlay/Source/Private/OnlineAsyncTaskGooglePlayFlushLeaderboards.cpp
@@ -0,0 +1,32 @@
+// Copyright Epic Games, Inc. All Rights Reserved.
+
+#include "OnlineAsyncTaskGooglePlayFlushLeaderboards.h"
+#include "GooglePlayGamesWrapper.h"
+#include "OnlineLeaderboardInterfaceGooglePlay.h"
+#include "OnlineSubsystemGooglePlay.h"
+
+FOnlineAsyncTaskGooglePlayFlushLeaderboards::FOnlineAsyncTaskGooglePlayFlushLeaderboards(
+	FOnlineSubsystemGooglePlay* InSubsystem,
+	FName InSessionName,
+	TArray<FGooglePlayLeaderboardScore>&& InLeaderboardScores)
+	: FOnlineAsyncTaskBasic(InSubsystem)
+	, SessionName(InSessionName)
+	, LeaderboardScores(MoveTemp(InLeaderboardScores))
+{
+}
+
+void FOnlineAsyncTaskGooglePlayFlushLeaderboards::Tick()
+{
+	if ( !bStarted)
+	{
+		bStarted = true;
+		bWasSuccessful = Subsystem->GetGooglePlayGamesWrapper().FlushLeaderboardsScores(this, MoveTemp(LeaderboardScores)); 
+		bIsComplete = !bWasSuccessful;
+	}
+}
+
+void FOnlineAsyncTaskGooglePlayFlushLeaderboards::TriggerDelegates()
+{
+	Subsystem->GetLeaderboardsGooglePlay()->TriggerOnLeaderboardFlushCompleteDelegates(SessionName, bWasSuccessful);
+}
+
diff --git a/Engine/Plugins/Online/Android/OnlineSubsystemGooglePlay/Source/Private/OnlineAsyncTaskGooglePlayFlushLeaderboards.h b/Engine/Plugins/Online/Android/OnlineSubsystemGooglePlay/Source/Private/OnlineAsyncTaskGooglePlayFlushLeaderboards.h
new file mode 100644
--- /dev/null
+++ b/Engine/Plugins/Online/Android/OnlineSubsystemGooglePlay/Source/Private/OnlineAsyncTaskGooglePlayFlushLeaderboards.h
@@ -0,0 +1,42 @@
+// Copyright Epic Games, Inc. All Rights Reserved.
+
+#pragma once
+
+#include "CoreMinimal.h"
+#include "OnlineAsyncTaskManager.h"
+#include "OnlineLeaderboardGooglePlayCommon.h"
+
+class FOnlineSubsystemGooglePlay;
+
+/** Task object to keep track of writing achievement score data 
+ * Objects of this type are associated with a GooglePlayGamesWrapper method call that routes operations to the Java backend. 
+ * The GooglePlayGamesWrapper implementation will adapt and set the task result and mark the task as completed when the 
+ * Java operation is finished 
+ */
+class FOnlineAsyncTaskGooglePlayFlushLeaderboards : public FOnlineAsyncTaskBasic<FOnlineSubsystemGooglePlay>
+{
+public:
+	/**
+	 * @brief Construct a new FOnlineAsyncTaskGooglePlayFlushLeaderboards object
+	 * 
+	 * @param InSubsystem Owning subsystem 
+	 * @param InSessionName Session name to report using in the delegate
+	 * @param LeaderboardScores Achievement data to submit 
+	 */
+	FOnlineAsyncTaskGooglePlayFlushLeaderboards( 
+		FOnlineSubsystemGooglePlay* InSubsystem, 
+		FName InSessionName, 
+		TArray<FGooglePlayLeaderboardScore>&& LeaderboardScores);
+
+	// FOnlineAsyncTask
+	virtual void Tick() override;
+	virtual void TriggerDelegates() override;
+
+	// FOnlineAsyncItem
+	virtual FString ToString() const override { return TEXT("FlushLeaderboards"); }
+
+private:
+	TArray<FGooglePlayLeaderboardScore> LeaderboardScores;
+	FName SessionName;
+	bool bStarted = false;
+};
diff --git a/Engine/Plugins/Online/Android/OnlineSubsystemGooglePlay/Source/Private/OnlineAsyncTaskGooglePlayLogin.cpp b/Engine/Plugins/Online/Android/OnlineSubsystemGooglePlay/Source/Private/OnlineAsyncTaskGooglePlayLogin.cpp
--- a/Engine/Plugins/Online/Android/OnlineSubsystemGooglePlay/Source/Private/OnlineAsyncTaskGooglePlayLogin.cpp
+++ b/Engine/Plugins/Online/Android/OnlineSubsystemGooglePlay/Source/Private/OnlineAsyncTaskGooglePlayLogin.cpp
@@ -1,108 +1,56 @@
 // Copyright Epic Games, Inc. All Rights Reserved.
 
 #include "OnlineAsyncTaskGooglePlayLogin.h"
+#include "GooglePlayGamesWrapper.h"
 #include "OnlineSubsystemGooglePlay.h"
-#include "AndroidRuntimeSettings.h"
 
-THIRD_PARTY_INCLUDES_START
-#include "gpg/builder.h"
-#include "gpg/debug.h"
-#include "gpg/types.h"
-THIRD_PARTY_INCLUDES_END
-
-FOnlineAsyncTaskGooglePlayLogin::FOnlineAsyncTaskGooglePlayLogin(
-	FOnlineSubsystemGooglePlay* InSubsystem,
-	int InPlayerId,
-	const FOnCompletedDelegate& InDelegate)
-	: FOnlineAsyncTaskGooglePlayAuthAction(InSubsystem)
-	, PlayerId(InPlayerId)
-	, Status(gpg::AuthStatus::ERROR_NOT_AUTHORIZED)
-	, Delegate(InDelegate)
+FOnlineAsyncTaskGooglePlayLogin::FOnlineAsyncTaskGooglePlayLogin(FOnlineSubsystemGooglePlay* InSubsystem, const FString& InAuthCodeClientId, bool InForceRefreshToken)
+	: FOnlineAsyncTaskBasic(InSubsystem)
+	, AuthCodeClientId(InAuthCodeClientId)
+	, ForceRefreshToken(InForceRefreshToken)
+	, ReceivedPlayerNetId(FUniqueNetIdGooglePlay::EmptyId())
 {
 }
 
-void FOnlineAsyncTaskGooglePlayLogin::Start_OnTaskThread()
+void FOnlineAsyncTaskGooglePlayLogin::Tick()
 {
-	UE_LOG_ONLINE(Log, TEXT("FOnlineAsyncTaskGooglePlayLogin::Start_OnTaskThread"));
-
-	// If we haven't created a GameServices object yet, do so.
-	if (Subsystem->GameServicesPtr.get() == nullptr)
+	if ( !bStarted)
 	{
-		UE_LOG_ONLINE(Log, TEXT("FOnlineAsyncTaskGooglePlayLogin::Start_OnTaskThread initializing game services"));
-		// Store the Subsystem pointer locally so that the OnAuthActionFinished lambda can capture it
-		FOnlineSubsystemGooglePlay* LocalSubsystem = Subsystem;
-
-		auto DefaultSettings = GetDefault<UAndroidRuntimeSettings>();
-
-		if (DefaultSettings->bEnableSnapshots)
-		{
-			Subsystem->GameServicesPtr = gpg::GameServices::Builder()
-				.SetDefaultOnLog(gpg::LogLevel::VERBOSE)
-				.SetOnAuthActionStarted([](gpg::AuthOperation Op) {
-					UE_LOG_ONLINE(Log, TEXT("GPG OnAuthActionStarted: %s"), *FString(DebugString(Op).c_str()));
-				})
-				.SetOnAuthActionFinished([LocalSubsystem](gpg::AuthOperation Op, gpg::AuthStatus LocalStatus) {
-					UE_LOG_ONLINE(Log, TEXT("GPG OnAuthActionFinished: %s, AuthStatus: %s"),
-						*FString(DebugString(Op).c_str()),
-						*FString(DebugString(LocalStatus).c_str()));
-					LocalSubsystem->OnAuthActionFinished(Op, LocalStatus);
-				})
-				.EnableSnapshots()
-				.Create(Subsystem->PlatformConfiguration);
-		}
-		else
-		{
-			Subsystem->GameServicesPtr = gpg::GameServices::Builder()
-				.SetDefaultOnLog(gpg::LogLevel::VERBOSE)
-				.SetOnAuthActionStarted([](gpg::AuthOperation Op) {
-					UE_LOG_ONLINE(Log, TEXT("GPG OnAuthActionStarted: %s"), *FString(DebugString(Op).c_str()));
-				})
-				.SetOnAuthActionFinished([LocalSubsystem](gpg::AuthOperation Op, gpg::AuthStatus LocalStatus) {
-					UE_LOG_ONLINE(Log, TEXT("GPG OnAuthActionFinished: %s, AuthStatus: %s"),
-						*FString(DebugString(Op).c_str()),
-						*FString(DebugString(LocalStatus).c_str()));
-					LocalSubsystem->OnAuthActionFinished(Op, LocalStatus);
-				})
-				.Create(Subsystem->PlatformConfiguration);
-		}
+		bStarted = true;
+		bWasSuccessful = Subsystem->GetGooglePlayGamesWrapper().Login(this, AuthCodeClientId, ForceRefreshToken); 
+		bIsComplete = !bWasSuccessful;
 	}
-	else if(Subsystem->GameServicesPtr->IsAuthorized())
+}
+
+void FOnlineAsyncTaskGooglePlayLogin::Finalize()
+{
+	FOnlineIdentityGooglePlayPtr IdentityInt = Subsystem->GetIdentityGooglePlay();
+	FUniqueNetIdPtr CurrentNetId = IdentityInt->GetUniquePlayerId(0);
+	bWasLoggedIn = CurrentNetId && CurrentNetId->IsValid();
+	if (bWasSuccessful)
 	{
-		UE_LOG_ONLINE(Log, TEXT("FOnlineAsyncTaskGooglePlayLogin::Start_OnTaskThread already authorized"));
-		// We have a GameServices object and the user is authorized, nothing else to do.
-		Status = gpg::AuthStatus::VALID;
-		bWasSuccessful = true;
-		bIsComplete = true;
+		IdentityInt->SetIdentityData(ReceivedPlayerNetId, MoveTemp(ReceivedDisplayName), MoveTemp(ReceivedAuthCode));
 	}
 	else
 	{
-		UE_LOG_ONLINE(Log, TEXT("FOnlineAsyncTaskGooglePlayLogin::Start_OnTaskThread not authorized"));
-		// We have created the GameServices object but the user isn't authorized.
-		bWasSuccessful = false;
-		bIsComplete = true;
+		IdentityInt->ClearIdentity();
 	}
 }
 
-void FOnlineAsyncTaskGooglePlayLogin::Finalize()
-{
-	UE_LOG_ONLINE(Log, TEXT("FOnlineAsyncTaskGooglePlayLogin: Finalize."));
-	// Async task manager owns the task and is responsible for cleaning it up.
-	Subsystem->CurrentLoginTask = nullptr;
-}
-
 void FOnlineAsyncTaskGooglePlayLogin::TriggerDelegates()
 {
-	UE_LOG_ONLINE(Log, TEXT("FOnlineAsyncTaskGooglePlayLogin: TriggerDelegates. %d"), WasSuccessful());
-	Delegate.ExecuteIfBound();
+	FOnlineIdentityGooglePlayPtr IdentityInt = Subsystem->GetIdentityGooglePlay();
+	IdentityInt->TriggerOnLoginCompleteDelegates(0, bWasSuccessful, *ReceivedPlayerNetId, TEXT(""));
+	if (!bWasLoggedIn)
+	{
+		IdentityInt->TriggerOnLoginStatusChangedDelegates(0, ELoginStatus::NotLoggedIn, ELoginStatus::LoggedIn, *ReceivedPlayerNetId);
+		IdentityInt->TriggerOnLoginChangedDelegates(0);
+	}
 }
 
-void FOnlineAsyncTaskGooglePlayLogin::OnAuthActionFinished(gpg::AuthOperation InOp, gpg::AuthStatus InStatus)
+void FOnlineAsyncTaskGooglePlayLogin::SetLoginData(FString&& PlayerId, FString&& DisplayName, FString&& AuthCode)
 {
-	UE_LOG_ONLINE(Log, TEXT("FOnlineAsyncTaskGooglePlayLogin::OnAuthActionFinished %d %d"), (int32)InOp, (int32)Status);
-	if (InOp == gpg::AuthOperation::SIGN_IN)
-	{
-		Status = InStatus;
-		bWasSuccessful = Status == gpg::AuthStatus::VALID;
-		bIsComplete = true;
-	}
+	ReceivedPlayerNetId = FUniqueNetIdGooglePlay::Create(MoveTemp(PlayerId));
+	ReceivedDisplayName = MoveTemp(DisplayName);
+	ReceivedAuthCode = MoveTemp(AuthCode);
 }
diff --git a/Engine/Plugins/Online/Android/OnlineSubsystemGooglePlay/Source/Private/OnlineAsyncTaskGooglePlayLogin.h b/Engine/Plugins/Online/Android/OnlineSubsystemGooglePlay/Source/Private/OnlineAsyncTaskGooglePlayLogin.h
--- a/Engine/Plugins/Online/Android/OnlineSubsystemGooglePlay/Source/Private/OnlineAsyncTaskGooglePlayLogin.h
+++ b/Engine/Plugins/Online/Android/OnlineSubsystemGooglePlay/Source/Private/OnlineAsyncTaskGooglePlayLogin.h
@@ -4,44 +4,46 @@
 
 #include "CoreMinimal.h"
 #include "OnlineAsyncTaskManager.h"
-#include "OnlineAsyncTaskGooglePlayAuthAction.h"
-#include "OnlineSubsystemGooglePlayPackage.h"
-
-THIRD_PARTY_INCLUDES_START
-#include "gpg/status.h"
-#include "gpg/types.h"
-THIRD_PARTY_INCLUDES_END
+#include "OnlineIdentityInterfaceGooglePlay.h"
 
 class FOnlineSubsystemGooglePlay;
 
-class FOnlineAsyncTaskGooglePlayLogin : public FOnlineAsyncTaskGooglePlayAuthAction
+/** Task object to keep track of writing achievement score data 
+ * Objects of this type are associated with a GooglePlayGamesWrapper method call that routes operations to the Java backend. 
+ * The GooglePlayGamesWrapper implementation will adapt and set the task result and mark the task as completed when the 
+ * Java operation is finished 
+ */
+class FOnlineAsyncTaskGooglePlayLogin : public FOnlineAsyncTaskBasic<FOnlineSubsystemGooglePlay>
 {
 public:
-	/** Delegate fired upon completion. */
-	DECLARE_DELEGATE(FOnCompletedDelegate);
-
 	/**
-	 * Constructor.
-	 *
-	 * @param InSubsystem a pointer to the owning subsysetm
-	 * @param InPlayerId index of the player who's logging in
+	 * @brief Construct a new FOnlineAsyncTaskGooglePlayLogin object
+	 * 
+	 * @param InSubsystem Owning subsystem
+	 * @param InAuthCodeClientId The client ID of the server that will perform the authorization code flow exchange. Won't request authorization code if empty
+	 * @param InForceRefreshToken If true a refresh token will be included in addition to the access token when the authorization code is exchanged
 	 */
-	FOnlineAsyncTaskGooglePlayLogin(FOnlineSubsystemGooglePlay* InSubsystem, int InPlayerId, const FOnCompletedDelegate& InDelegate);
+	FOnlineAsyncTaskGooglePlayLogin(
+		FOnlineSubsystemGooglePlay* InSubsystem, 
+		const FString& InAuthCodeClientId, 
+		bool InForceRefreshToken);
+
+	// FOnlineAsyncTask
+	virtual void Tick() override;
+	void Finalize() override;
+	void TriggerDelegates() override;
 
 	// FOnlineAsyncItem
 	virtual FString ToString() const override { return TEXT("Login"); }
-	virtual void Finalize() override;
-	virtual void TriggerDelegates() override;
 
+	// Set task result data. Accessed trhough GooglePlayGamesWrapper implementation
+	void SetLoginData(FString&& PlayerId, FString&& PlayerDisplayName, FString&& AuthCode);
 private:
-	/** The subsystem is the only class that should be calling OnAuthActionFinished */
-	friend class FOnlineSubsystemGooglePlay;
-
-	// FOnlineAsyncTaskGooglePlayAuthAction
-	virtual void OnAuthActionFinished(gpg::AuthOperation InOp, gpg::AuthStatus InStatus) override;
-	virtual void Start_OnTaskThread() override;
-
-	int PlayerId;
-	gpg::AuthStatus Status;
-	FOnCompletedDelegate Delegate;
+	FString AuthCodeClientId;
+	bool ForceRefreshToken;
+	bool bStarted = false;
+	bool bWasLoggedIn = false;
+	FUniqueNetIdGooglePlayPtr ReceivedPlayerNetId;
+	FString ReceivedDisplayName;
+	FString ReceivedAuthCode;
 };
diff --git a/Engine/Plugins/Online/Android/OnlineSubsystemGooglePlay/Source/Private/OnlineAsyncTaskGooglePlayLogout.cpp b/Engine/Plugins/Online/Android/OnlineSubsystemGooglePlay/Source/Private/OnlineAsyncTaskGooglePlayLogout.cpp
deleted file mode 100644
diff --git a/Engine/Plugins/Online/Android/OnlineSubsystemGooglePlay/Source/Private/OnlineAsyncTaskGooglePlayLogout.h b/Engine/Plugins/Online/Android/OnlineSubsystemGooglePlay/Source/Private/OnlineAsyncTaskGooglePlayLogout.h
deleted file mode 100644
diff --git a/Engine/Plugins/Online/Android/OnlineSubsystemGooglePlay/Source/Private/OnlineAsyncTaskGooglePlayQueryAchievements.cpp b/Engine/Plugins/Online/Android/OnlineSubsystemGooglePlay/Source/Private/OnlineAsyncTaskGooglePlayQueryAchievements.cpp
--- a/Engine/Plugins/Online/Android/OnlineSubsystemGooglePlay/Source/Private/OnlineAsyncTaskGooglePlayQueryAchievements.cpp
+++ b/Engine/Plugins/Online/Android/OnlineSubsystemGooglePlay/Source/Private/OnlineAsyncTaskGooglePlayQueryAchievements.cpp
@@ -4,44 +4,47 @@
 #include "OnlineSubsystemGooglePlay.h"
 #include "OnlineAchievementsInterfaceGooglePlay.h"
 
-THIRD_PARTY_INCLUDES_START
-#include "gpg/achievement_manager.h"
-THIRD_PARTY_INCLUDES_END
-
 FOnlineAsyncTaskGooglePlayQueryAchievements::FOnlineAsyncTaskGooglePlayQueryAchievements(
 	FOnlineSubsystemGooglePlay* InSubsystem,
-	const FUniqueNetIdGooglePlay& InUserId,
+	const FUniqueNetIdPtr& InPlayerId,
 	const FOnQueryAchievementsCompleteDelegate& InDelegate)
 	: FOnlineAsyncTaskBasic(InSubsystem)
-	, UserId(InUserId.AsShared())
+	, PlayerId(InPlayerId)
 	, Delegate(InDelegate)
 {
 }
 
+void FOnlineAsyncTaskGooglePlayQueryAchievements::SetAchievementsData(TArray<FOnlineAchievementGooglePlay>&& InAchievementsData, TArray<FOnlineAchievementDesc>&& InAchievementsDesc)
+{
+	AchievementsData = MoveTemp(InAchievementsData);	
+	AchievementsDesc = MoveTemp(InAchievementsDesc);
+}
+
+void FOnlineAsyncTaskGooglePlayQueryAchievements::Tick()
+{
+	if ( !bStarted)
+	{
+		bStarted = true;
+		bWasSuccessful = Subsystem->GetGooglePlayGamesWrapper().QueryAchievements(this); 
+		bIsComplete = !bWasSuccessful;
+	}
+}
+
 void FOnlineAsyncTaskGooglePlayQueryAchievements::Finalize()
 {
+	FOnlineAchievementsGooglePlayPtr AchievementsInt = Subsystem->GetAchievementsGooglePlay();
+
 	if (bWasSuccessful)
 	{
-		Subsystem->GetAchievementsGooglePlay()->UpdateCache(Response);
+		AchievementsInt->UpdateCache(MoveTemp(AchievementsData), MoveTemp(AchievementsDesc));
 	}
 	else
 	{
-		Subsystem->GetAchievementsGooglePlay()->ClearCache();
+		AchievementsInt->ClearCache();
 	}
 }
 
 void FOnlineAsyncTaskGooglePlayQueryAchievements::TriggerDelegates()
 {
-	Delegate.ExecuteIfBound(*UserId, bWasSuccessful);
-}
-
-void FOnlineAsyncTaskGooglePlayQueryAchievements::Tick()
-{
-	// We're already running on the online thread. Using the blocking version of the API function
-	// won't block the main thread and simplifies things a bit.
-	// Try a 10 second timeout.
-	Response = Subsystem->GetGameServices()->Achievements().FetchAllBlocking(gpg::DataSource::CACHE_OR_NETWORK, gpg::Timeout(10000));
-	
-	bWasSuccessful = Response.status == gpg::ResponseStatus::VALID || Response.status == gpg::ResponseStatus::VALID_BUT_STALE;
-	bIsComplete = true;
-}
+	Delegate.ExecuteIfBound(*PlayerId, bWasSuccessful);
+}
\ No newline at end of file
diff --git a/Engine/Plugins/Online/Android/OnlineSubsystemGooglePlay/Source/Private/OnlineAsyncTaskGooglePlayQueryAchievements.h b/Engine/Plugins/Online/Android/OnlineSubsystemGooglePlay/Source/Private/OnlineAsyncTaskGooglePlayQueryAchievements.h
--- a/Engine/Plugins/Online/Android/OnlineSubsystemGooglePlay/Source/Private/OnlineAsyncTaskGooglePlayQueryAchievements.h
+++ b/Engine/Plugins/Online/Android/OnlineSubsystemGooglePlay/Source/Private/OnlineAsyncTaskGooglePlayQueryAchievements.h
@@ -4,34 +4,45 @@
 
 #include "CoreMinimal.h"
 #include "OnlineAsyncTaskManager.h"
+#include "OnlineAchievementGooglePlayCommon.h"
 #include "OnlineSubsystemGooglePlayPackage.h"
-#include "Interfaces/OnlineAchievementsInterface.h"
-#include "OnlineIdentityInterfaceGooglePlay.h"
-
-THIRD_PARTY_INCLUDES_START
-#include "gpg/achievement_manager.h"
-THIRD_PARTY_INCLUDES_END
 
 class FOnlineSubsystemGooglePlay;
 
+/** Task object to keep track of writing achievement score data 
+ * Objects of this type are associated with a GooglePlayGamesWrapper method call that routes operations to the Java backend. 
+ * The GooglePlayGamesWrapper implementation will adapt and set the task result and mark the task as completed when the 
+ * Java operation is finished 
+ */
 class FOnlineAsyncTaskGooglePlayQueryAchievements : public FOnlineAsyncTaskBasic<FOnlineSubsystemGooglePlay>
 {
 public:
+	/**
+	 * @brief Construct a new FOnlineAsyncTaskGooglePlayQueryAchievements object
+	 * 
+	 * @param InSubsystem Owning subsystem
+	 * @param PlayerId PlayerId to report in the delegate. Only data for local player can be received from GooglePlay
+	 * @param InDelegate Delegate to invoke on task completion
+	 */
 	FOnlineAsyncTaskGooglePlayQueryAchievements(
 		FOnlineSubsystemGooglePlay* InSubsystem,
-		const FUniqueNetIdGooglePlay& InUserId,
+		const FUniqueNetIdPtr& PlayerId,
 		const FOnQueryAchievementsCompleteDelegate& InDelegate);
 
-	// FOnlineAsyncItem
-	virtual FString ToString() const override { return TEXT("QueryAchievements"); }
+	// FOnlineAsyncTask
+	virtual void Tick() override;
 	virtual void Finalize() override;
 	virtual void TriggerDelegates() override;
 
-	// FOnlineAsyncTask
-	virtual void Tick() override;
+	// FOnlineAsyncItem
+	virtual FString ToString() const override { return TEXT("QueryAchievements"); }
 
+	// Set task result data. Accessed trhough GooglePlayGamesWrapper implementation
+	void SetAchievementsData(TArray<FOnlineAchievementGooglePlay>&& InAchievementsData, TArray<FOnlineAchievementDesc>&& InAchievementsDesc);
 private:
-	FUniqueNetIdGooglePlayRef UserId;
+	FUniqueNetIdPtr PlayerId;
 	FOnQueryAchievementsCompleteDelegate Delegate;
-	gpg::AchievementManager::FetchAllResponse Response;
+	TArray<FOnlineAchievementGooglePlay> AchievementsData;
+	TArray<FOnlineAchievementDesc> AchievementsDesc;
+	bool bStarted = false;
 };
diff --git a/Engine/Plugins/Online/Android/OnlineSubsystemGooglePlay/Source/Private/OnlineAsyncTaskGooglePlayReadLeaderboard.cpp b/Engine/Plugins/Online/Android/OnlineSubsystemGooglePlay/Source/Private/OnlineAsyncTaskGooglePlayReadLeaderboard.cpp
--- a/Engine/Plugins/Online/Android/OnlineSubsystemGooglePlay/Source/Private/OnlineAsyncTaskGooglePlayReadLeaderboard.cpp
+++ b/Engine/Plugins/Online/Android/OnlineSubsystemGooglePlay/Source/Private/OnlineAsyncTaskGooglePlayReadLeaderboard.cpp
@@ -1,95 +1,76 @@
 // Copyright Epic Games, Inc. All Rights Reserved.
 
 #include "OnlineAsyncTaskGooglePlayReadLeaderboard.h"
-#include "OnlineSubsystemGooglePlay.h"
-#include "OnlineIdentityInterfaceGooglePlay.h"
+#include "GooglePlayGamesWrapper.h"
 #include "OnlineLeaderboardInterfaceGooglePlay.h"
-
-THIRD_PARTY_INCLUDES_START
-#include "gpg/leaderboard_manager.h"
-THIRD_PARTY_INCLUDES_END
-
-using namespace gpg;
+#include "OnlineSubsystemGooglePlay.h"
 
 FOnlineAsyncTaskGooglePlayReadLeaderboard::FOnlineAsyncTaskGooglePlayReadLeaderboard(
 	FOnlineSubsystemGooglePlay* InSubsystem,
-	FOnlineLeaderboardReadRef& InReadObject,
-	const FString& InLeaderboardId)
+	const FOnlineLeaderboardReadRef& InReadObject,
+	const FString& InLeaderboardId )
 	: FOnlineAsyncTaskBasic(InSubsystem)
 	, ReadObject(InReadObject)
 	, LeaderboardId(InLeaderboardId)
 {
-	Response.status = ResponseStatus::ERROR_TIMEOUT;
+	ReadObject->Rows.Empty();
+	ReadObject->ReadState = EOnlineAsyncTaskState::InProgress;
+}
+
+void FOnlineAsyncTaskGooglePlayReadLeaderboard::Tick()
+{
+	if ( !bStarted)
+	{
+		bStarted = true;
+		bWasSuccessful = Subsystem->GetGooglePlayGamesWrapper().RequestLeaderboardScore(this, LeaderboardId); 
+		bIsComplete = !bWasSuccessful;
+	}
 }
 
 void FOnlineAsyncTaskGooglePlayReadLeaderboard::Finalize()
 {
-	ReadObject->ReadState = bWasSuccessful ? EOnlineAsyncTaskState::Done : EOnlineAsyncTaskState::Failed;
+	ReadObject->ReadState = bWasSuccessful? EOnlineAsyncTaskState::Done : EOnlineAsyncTaskState::Failed;
 }
 
 void FOnlineAsyncTaskGooglePlayReadLeaderboard::TriggerDelegates()
 {
-	Subsystem->GetLeaderboardsInterface()->TriggerOnLeaderboardReadCompleteDelegates(bWasSuccessful);
+	Subsystem->GetLeaderboardsGooglePlay()->TriggerOnLeaderboardReadCompleteDelegates(bWasSuccessful);
 }
 
-void FOnlineAsyncTaskGooglePlayReadLeaderboard::Tick()
+void FOnlineAsyncTaskGooglePlayReadLeaderboard::AddScore(const FString& DisplayName, const FString& PlayerId, int64 Rank, int64 RawScore)
 {
-	// Convert our FString leaderboard ID to ASCII for the API's std::string.
-	auto ConvertedId = FOnlineSubsystemGooglePlay::ConvertFStringToStdString(LeaderboardId);
-	Response = Subsystem->GetGameServices()->Leaderboards().FetchScoreSummaryBlocking(
-		DataSource::CACHE_OR_NETWORK,
-		Timeout(10000),
-		ConvertedId,
-		LeaderboardTimeSpan::ALL_TIME,
-		LeaderboardCollection::PUBLIC);
-
-	bWasSuccessful = false;
-
-	if (Response.status != ResponseStatus::VALID)
-	{
-		bIsComplete = true;
-		return;
-	}
-
-	// We can only get the current user's leaderboard score from Google Play, so just add one row with it.
-	//FUniqueNetIdRef UserId = FUniqueNetIdGooglePlay::Create(TEXT("0"));
-	auto UserId = Subsystem->GetIdentityGooglePlay()->GetCurrentUserId();
-	if (!UserId.IsValid())
+	if (!PlayerId.IsEmpty())
 	{
-		// If there's no user signed in, can't get leaderboard.
-		bIsComplete = true;
-		return;
-	}
-
-	FOnlineStatsRow* UserRow = ReadObject.Get().FindPlayerRecord(*UserId);
-	if (UserRow == NULL)
-	{
-		UserRow = new (ReadObject->Rows) FOnlineStatsRow(UserId->ToString(), UserId.ToSharedRef());
-	}
+		FUniqueNetIdRef UserId = FUniqueNetIdGooglePlay::Create(PlayerId);
+		FOnlineStatsRow* UserRow = ReadObject->FindPlayerRecord(*UserId);
+		if (UserRow == NULL)
+		{
+			UserRow = &ReadObject->Rows.Emplace_GetRef(DisplayName, UserId);
+		}
 
-	for (int32 StatIdx = 0; StatIdx < ReadObject->ColumnMetadata.Num(); StatIdx++)
-	{
-		const FColumnMetaData& ColumnMeta = ReadObject->ColumnMetadata[StatIdx];
+		UserRow->Rank = Rank;
 
-		FVariantData* LastColumn = NULL;
-		switch (ColumnMeta.DataType)
+		for (const FColumnMetaData& ColumnMeta: ReadObject->ColumnMetadata)
 		{
-			case EOnlineKeyValuePairDataType::Int32:
+			switch (ColumnMeta.DataType)
 			{
-				int32 Value = Response.data.CurrentPlayerScore().Value();
-				LastColumn = &(UserRow->Columns.Add(ColumnMeta.ColumnName, FVariantData(Value)));
-				bWasSuccessful = true;
-				break;
-			}
+				case EOnlineKeyValuePairDataType::Int64:
+				{
+					UserRow->Columns.Add(ColumnMeta.ColumnName, FVariantData(RawScore));
+					break;
+				}
+				case EOnlineKeyValuePairDataType::Int32:
+				{
+					int32 Value = (int32)FMath::Clamp(RawScore, MIN_int32, MAX_int32);
+					UserRow->Columns.Add(ColumnMeta.ColumnName, FVariantData(Value));
+					break;
+				}
 
-			default:
-			{
-				UE_LOG_ONLINE(Warning, TEXT("Unsupported key value pair during retrieval from Google Play %s"), *ColumnMeta.ColumnName.ToString());
-				break;
+				default:
+				{
+					UE_LOG_ONLINE(Warning, TEXT("Unsupported key value pair during retrieval from Google Play %s"), *ColumnMeta.ColumnName.ToString());
+					break;
+				}
 			}
 		}
-	}
-	
-	bIsComplete = true;
-}
-
+	}}
\ No newline at end of file
diff --git a/Engine/Plugins/Online/Android/OnlineSubsystemGooglePlay/Source/Private/OnlineAsyncTaskGooglePlayReadLeaderboard.h b/Engine/Plugins/Online/Android/OnlineSubsystemGooglePlay/Source/Private/OnlineAsyncTaskGooglePlayReadLeaderboard.h
--- a/Engine/Plugins/Online/Android/OnlineSubsystemGooglePlay/Source/Private/OnlineAsyncTaskGooglePlayReadLeaderboard.h
+++ b/Engine/Plugins/Online/Android/OnlineSubsystemGooglePlay/Source/Private/OnlineAsyncTaskGooglePlayReadLeaderboard.h
@@ -3,38 +3,43 @@
 #pragma once
 
 #include "CoreMinimal.h"
-#include "OnlineAsyncTaskManager.h"
 #include "OnlineStats.h"
-
-THIRD_PARTY_INCLUDES_START
-#include "gpg/leaderboard_manager.h"
-THIRD_PARTY_INCLUDES_END
+#include "OnlineAsyncTaskManager.h"
 
 class FOnlineSubsystemGooglePlay;
 
+/** Task object to keep track of writing achievement score data 
+ * Objects of this type are associated with a GooglePlayGamesWrapper method call that routes operations to the Java backend. 
+ * The GooglePlayGamesWrapper implementation will adapt and set the task result and mark the task as completed when the 
+ * Java operation is finished 
+ */
 class FOnlineAsyncTaskGooglePlayReadLeaderboard : public FOnlineAsyncTaskBasic<FOnlineSubsystemGooglePlay>
 {
 public:
-	FOnlineAsyncTaskGooglePlayReadLeaderboard(
-		FOnlineSubsystemGooglePlay* InSubsystem,
-		FOnlineLeaderboardReadRef& InReadObject,
+	/**
+	 * @brief Construct a new FOnlineAsyncTaskGooglePlayReadLeaderboard object
+	 * 
+	 * @param InSubsystem Owning subsystem 
+	 * @param InReadObject Destination object for the leaderboard information received
+	 * @param InLeaderboardId GooglePlay leaderboard id to request
+	 */
+	FOnlineAsyncTaskGooglePlayReadLeaderboard( 
+		FOnlineSubsystemGooglePlay* InSubsystem, 
+		const FOnlineLeaderboardReadRef& InReadObject, 
 		const FString& InLeaderboardId);
 
-	// FOnlineAsyncItem
-	virtual FString ToString() const override { return TEXT("ReadLeaderboard"); }
-	virtual void Finalize() override;
-	virtual void TriggerDelegates() override;
-
 	// FOnlineAsyncTask
 	virtual void Tick() override;
+	void Finalize() override;
+	void TriggerDelegates() override;
 
-private:
-	/** Leaderboard read data */
-	FOnlineLeaderboardReadRef ReadObject;
+	// FOnlineAsyncItem
+	virtual FString ToString() const override { return TEXT("ReadLeaderboard"); }
 
-	/** Google Play leaderboard id */
+	// Set task result data. Accessed trhough GooglePlayGamesWrapper implementation
+	void AddScore(const FString& DisplayName, const FString& PlayerId, int64 Rank, int64 RawScore);
+private:
 	FString LeaderboardId;
-
-	/** API query result */
-	gpg::LeaderboardManager::FetchScoreSummaryResponse Response;
+	bool bStarted = false;
+	FOnlineLeaderboardReadRef ReadObject;
 };
diff --git a/Engine/Plugins/Online/Android/OnlineSubsystemGooglePlay/Source/Private/OnlineAsyncTaskGooglePlayShowLoginUI.cpp b/Engine/Plugins/Online/Android/OnlineSubsystemGooglePlay/Source/Private/OnlineAsyncTaskGooglePlayShowLoginUI.cpp
deleted file mode 100644
diff --git a/Engine/Plugins/Online/Android/OnlineSubsystemGooglePlay/Source/Private/OnlineAsyncTaskGooglePlayShowLoginUI.h b/Engine/Plugins/Online/Android/OnlineSubsystemGooglePlay/Source/Private/OnlineAsyncTaskGooglePlayShowLoginUI.h
deleted file mode 100644
diff --git a/Engine/Plugins/Online/Android/OnlineSubsystemGooglePlay/Source/Private/OnlineAsyncTaskGooglePlayWriteAchievements.cpp b/Engine/Plugins/Online/Android/OnlineSubsystemGooglePlay/Source/Private/OnlineAsyncTaskGooglePlayWriteAchievements.cpp
new file mode 100644
--- /dev/null
+++ b/Engine/Plugins/Online/Android/OnlineSubsystemGooglePlay/Source/Private/OnlineAsyncTaskGooglePlayWriteAchievements.cpp
@@ -0,0 +1,48 @@
+// Copyright Epic Games, Inc. All Rights Reserved.
+
+#include "OnlineAsyncTaskGooglePlayWriteAchievements.h"
+#include "Algo/NoneOf.h"
+#include "OnlineSubsystemGooglePlay.h"
+#include "OnlineAchievementsInterfaceGooglePlay.h"
+
+FOnlineAsyncTaskGooglePlayWriteAchievements::FOnlineAsyncTaskGooglePlayWriteAchievements(
+	FOnlineSubsystemGooglePlay* InSubsystem,
+	const FUniqueNetIdPtr& InPlayerId,
+	TArray<FGooglePlayAchievementWriteData>&& InWriteAchievements,
+	const FOnAchievementsWrittenDelegate& InDelegate)
+	: FOnlineAsyncTaskBasic(InSubsystem)
+	, PlayerId(InPlayerId)
+	, WriteAchievementsData(MoveTemp(InWriteAchievements))
+	, Delegate(InDelegate)
+{
+}
+
+void FOnlineAsyncTaskGooglePlayWriteAchievements::Tick()
+{
+	if ( !bStarted)
+	{
+		bStarted = true;
+		bWasSuccessful = Subsystem->GetGooglePlayGamesWrapper().WriteAchievements(this, WriteAchievementsData); 
+		bIsComplete = !bWasSuccessful;
+	}
+}
+
+bool FOnlineAsyncTaskGooglePlayWriteAchievements::SetSucceeded(const TArray<FString>& SucceededIds)
+{
+	// Remove data of entried that did not succeed
+	return 0 != WriteAchievementsData.RemoveAll([&SucceededIds](const FGooglePlayAchievementWriteData& Entry)
+		{
+			return Algo::NoneOf(SucceededIds, [&Entry](const FString& Id) { return Id == Entry.GooglePlayAchievementId;} );
+		});
+}
+
+void FOnlineAsyncTaskGooglePlayWriteAchievements::Finalize()
+{
+	FOnlineAchievementsGooglePlayPtr AchievementsInt = Subsystem->GetAchievementsGooglePlay();
+	AchievementsInt->UpdateCacheAfterWrite(WriteAchievementsData);
+}
+
+void FOnlineAsyncTaskGooglePlayWriteAchievements::TriggerDelegates()
+{
+	Delegate.ExecuteIfBound(*PlayerId, bWasSuccessful);
+}
\ No newline at end of file
diff --git a/Engine/Plugins/Online/Android/OnlineSubsystemGooglePlay/Source/Private/OnlineAsyncTaskGooglePlayWriteAchievements.h b/Engine/Plugins/Online/Android/OnlineSubsystemGooglePlay/Source/Private/OnlineAsyncTaskGooglePlayWriteAchievements.h
new file mode 100644
--- /dev/null
+++ b/Engine/Plugins/Online/Android/OnlineSubsystemGooglePlay/Source/Private/OnlineAsyncTaskGooglePlayWriteAchievements.h
@@ -0,0 +1,56 @@
+// Copyright Epic Games, Inc. All Rights Reserved.
+
+#pragma once
+
+#include "CoreMinimal.h"
+#include "OnlineAsyncTaskManager.h"
+#include "OnlineAchievementGooglePlayCommon.h"
+#include "Interfaces/OnlineAchievementsInterface.h"
+
+class FOnlineSubsystemGooglePlay;
+
+/** Task object to keep track of writing achievement score data 
+ * Objects of this type are associated with a GooglePlayGamesWrapper method call that routes operations to the Java backend. 
+ * The GooglePlayGamesWrapper implementation will adapt and set the task result and mark the task as completed when the 
+ * Java operation is finished 
+ */
+class FOnlineAsyncTaskGooglePlayWriteAchievements : public FOnlineAsyncTaskBasic<FOnlineSubsystemGooglePlay>
+{
+public:
+	/**
+	 * @brief Construct a new FOnlineAsyncTaskGooglePlayWriteAchievements object
+	 * 
+	 * @param InSubsystem Owning subsystem
+	 * @param PlayerId PlayerId to report in the delegate. Only data for local player can be received from GooglePlay
+	 * @param InWriteAchievements Achievement data to send
+	 * @param Delegate Delegate to invoke on task completion
+	 */
+	FOnlineAsyncTaskGooglePlayWriteAchievements(
+		FOnlineSubsystemGooglePlay* InSubsystem,
+		const FUniqueNetIdPtr& PlayerId,
+		TArray<FGooglePlayAchievementWriteData>&& InWriteAchievements,
+		const FOnAchievementsWrittenDelegate& Delegate);
+
+	/**
+	 * @brief Notifies the list of achievement writes that succeeded
+	 * 
+	 * @param SucceededIds 
+	 * @return true If all achievements expected to be written succeeded
+	 * @return false Otherwise
+	 */
+	bool SetSucceeded(const TArray<FString>& SucceededIds);
+
+	// FOnlineAsyncTask
+	virtual void Tick() override;
+	virtual void Finalize() override;
+	virtual void TriggerDelegates() override;
+
+	// FOnlineAsyncItem
+	virtual FString ToString() const override { return TEXT("QueryAchievements"); }
+
+private:
+	FUniqueNetIdPtr PlayerId;
+	TArray<FGooglePlayAchievementWriteData> WriteAchievementsData;
+	FOnAchievementsWrittenDelegate Delegate;
+	bool bStarted = false;
+};
diff --git a/Engine/Plugins/Online/Android/OnlineSubsystemGooglePlay/Source/Private/OnlineExternalUIInterfaceGooglePlay.cpp b/Engine/Plugins/Online/Android/OnlineSubsystemGooglePlay/Source/Private/OnlineExternalUIInterfaceGooglePlay.cpp
--- a/Engine/Plugins/Online/Android/OnlineSubsystemGooglePlay/Source/Private/OnlineExternalUIInterfaceGooglePlay.cpp
+++ b/Engine/Plugins/Online/Android/OnlineSubsystemGooglePlay/Source/Private/OnlineExternalUIInterfaceGooglePlay.cpp
@@ -5,11 +5,6 @@
 #include "UObject/Class.h"
 #include "OnlineSubsystemGooglePlay.h"
 
-THIRD_PARTY_INCLUDES_START
-#include "gpg/achievement_manager.h"
-#include "gpg/leaderboard_manager.h"
-THIRD_PARTY_INCLUDES_END
-
 FOnlineExternalUIGooglePlay::FOnlineExternalUIGooglePlay(FOnlineSubsystemGooglePlay* InSubsystem)
 	: Subsystem(InSubsystem)
 {
@@ -18,7 +13,26 @@ FOnlineExternalUIGooglePlay::FOnlineExternalUIGooglePlay(FOnlineSubsystemGoogleP
 
 bool FOnlineExternalUIGooglePlay::ShowLoginUI(const int ControllerIndex, bool bShowOnlineOnly, bool bShowSkipButton, const FOnLoginUIClosedDelegate& Delegate)
 {
-	Subsystem->StartShowLoginUITask(ControllerIndex, Delegate);	
+	IOnlineIdentityPtr OnlineIdentity = Subsystem->GetIdentityInterface();
+
+	if (FUniqueNetIdPtr UniqueNetId = OnlineIdentity->GetUniquePlayerId(ControllerIndex); UniqueNetId->IsValid())
+	{
+		Delegate.ExecuteIfBound(UniqueNetId, ControllerIndex, FOnlineError::Success());
+		return true;
+	}
+
+	TSharedPtr<FDelegateHandle> DelegateHandle = MakeShared<FDelegateHandle>();
+	*DelegateHandle = OnlineIdentity->AddOnLoginCompleteDelegate_Handle(ControllerIndex, FOnLoginCompleteDelegate::CreateLambda([this, DelegateHandle, Delegate](int32 LocalUserNum, bool bWasSuccessful, const FUniqueNetId& UserId, const FString& ErrorString)
+		{
+			FOnlineError Error(bWasSuccessful);
+			Error.SetFromErrorCode(ErrorString);
+
+			Delegate.ExecuteIfBound(UserId.IsValid()? UserId.AsShared() : FUniqueNetIdPtr(), LocalUserNum, Error);
+
+			IOnlineIdentityPtr OnlineIdentity = Subsystem->GetIdentityInterface();
+			OnlineIdentity->ClearOnLoginCompleteDelegate_Handle(LocalUserNum, *DelegateHandle);
+		}));
+	OnlineIdentity->Login(ControllerIndex, FOnlineAccountCredentials());
 	return true;
 }
 
@@ -34,46 +48,23 @@ bool FOnlineExternalUIGooglePlay::ShowInviteUI(int32 LocalUserNum, FName Session
 
 bool FOnlineExternalUIGooglePlay::ShowAchievementsUI(int32 LocalUserNum) 
 {
-	if (Subsystem->GetGameServices() == nullptr)
-	{
-		return false;
-	}
-
-	if (!Subsystem->GetGameServices()->IsAuthorized())
-	{
-		return false;
-	}
-
-	Subsystem->GetGameServices()->Achievements().ShowAllUI(nullptr);
+	Subsystem->GetGooglePlayGamesWrapper().ShowAchievementsUI();
 	return true;
 }
 
 bool FOnlineExternalUIGooglePlay::ShowLeaderboardUI(const FString& LeaderboardName)
 {
-	if (Subsystem->GetGameServices() == nullptr)
-	{
-		return false;
-	}
+	auto Settings = GetDefault<UAndroidRuntimeSettings>();
 
-	if (!Subsystem->GetGameServices()->IsAuthorized())
-	{
-		return false;
-	}
-
-	auto DefaultSettings = GetDefault<UAndroidRuntimeSettings>();
-
-	for(const auto& Mapping : DefaultSettings->LeaderboardMap)
+	for(const auto& Mapping : Settings->LeaderboardMap)
 	{
 		if(Mapping.Name == LeaderboardName)
 		{
-			auto ConvertedId = FOnlineSubsystemGooglePlay::ConvertFStringToStdString(Mapping.LeaderboardID);
-			Subsystem->GetGameServices()->Leaderboards().ShowUI(ConvertedId, nullptr);
-
+			Subsystem->GetGooglePlayGamesWrapper().ShowLeaderboardUI(Mapping.LeaderboardID);
 			return true;
 		}
 	}
-	
-	return false;
+	return true;
 }
 
 bool FOnlineExternalUIGooglePlay::ShowWebURL(const FString& Url, const FShowWebUrlParams& ShowParams, const FOnShowWebUrlClosedDelegate& Delegate)
diff --git a/Engine/Plugins/Online/Android/OnlineSubsystemGooglePlay/Source/Private/OnlineIdentityInterfaceGooglePlay.cpp b/Engine/Plugins/Online/Android/OnlineSubsystemGooglePlay/Source/Private/OnlineIdentityInterfaceGooglePlay.cpp
--- a/Engine/Plugins/Online/Android/OnlineSubsystemGooglePlay/Source/Private/OnlineIdentityInterfaceGooglePlay.cpp
+++ b/Engine/Plugins/Online/Android/OnlineSubsystemGooglePlay/Source/Private/OnlineIdentityInterfaceGooglePlay.cpp
@@ -1,22 +1,24 @@
 // Copyright Epic Games, Inc. All Rights Reserved.
 
 #include "OnlineIdentityInterfaceGooglePlay.h"
+#include "AndroidRuntimeSettings.h"
+#include "OnlineAsyncTaskGooglePlayLogin.h"
 #include "OnlineSubsystemGooglePlay.h"
 
-FOnlineIdentityGooglePlay::FPendingConnection FOnlineIdentityGooglePlay::PendingConnectRequest;
-
 FOnlineIdentityGooglePlay::FOnlineIdentityGooglePlay(FOnlineSubsystemGooglePlay* InSubsystem)
-	: bPrevLoggedIn(false)
-	, bLoggedIn(false)
-	, PlayerAlias("")
-	, AuthToken("NONE")
-	, MainSubsystem(InSubsystem)
-	, bRegisteringUser(false)
-	, bLoggingInUser(false)
+	: MainSubsystem(InSubsystem)
 {
 	UE_LOG_ONLINE_IDENTITY(Display, TEXT("FOnlineIdentityGooglePlay::FOnlineIdentityGooglePlay()"));
 	check(MainSubsystem != nullptr);
-	PendingConnectRequest.ConnectionInterface = this;
+
+	ClearIdentity();
+}
+
+void FOnlineIdentityGooglePlay::ClearIdentity()
+{
+	UniqueNetId = FUniqueNetIdGooglePlay::EmptyId();
+	PlayerAlias.Empty();
+	AuthCode.Empty();
 }
 
 TSharedPtr<FUserOnlineAccount> FOnlineIdentityGooglePlay::GetUserAccount(const FUniqueNetId& UserId) const
@@ -24,7 +26,6 @@ TSharedPtr<FUserOnlineAccount> FOnlineIdentityGooglePlay::GetUserAccount(const F
 	return nullptr;
 }
 
-
 TArray<TSharedPtr<FUserOnlineAccount> > FOnlineIdentityGooglePlay::GetAllUserAccounts() const
 {
 	TArray<TSharedPtr<FUserOnlineAccount> > Result;
@@ -34,74 +35,79 @@ TArray<TSharedPtr<FUserOnlineAccount> > FOnlineIdentityGooglePlay::GetAllUserAcc
 
 bool FOnlineIdentityGooglePlay::Login(int32 LocalUserNum, const FOnlineAccountCredentials& AccountCredentials) 
 {
-	bool bStartedLogin = false;
-	if (bLoggedIn)
-	{
-		// already logged in so just report all is ok!
-		// Now logged in
-		bStartedLogin = true;
-
-		static const int32 MAX_TEXT_LINE_LEN = 32;
-		TCHAR Line[MAX_TEXT_LINE_LEN + 1] = { 0 };
-		int32 Len = FCString::Snprintf(Line, MAX_TEXT_LINE_LEN, TEXT("%d"), LocalUserNum);
-
-		const FString PlayerId(Line);
-		UniqueNetId = FUniqueNetIdGooglePlay::Create(PlayerId);
-		TriggerOnLoginCompleteDelegates(LocalUserNum, true, *UniqueNetId, TEXT(""));
-	}
-	else if (!PendingConnectRequest.IsConnectionPending)
+	if (LocalUserNum > 0)
 	{
-		// Kick the login sequence...
-		bStartedLogin = true;
-		PendingConnectRequest.IsConnectionPending = true;
+		TriggerOnLoginCompleteDelegates(LocalUserNum, false, *FUniqueNetIdGooglePlay::EmptyId(), FString("FOnlineIdentityGooglePlay does not support more than 1 local player"));
+		return false;
 	}
-	else
+	
+	if (UniqueNetId->IsValid())
 	{
-		TriggerOnLoginCompleteDelegates(LocalUserNum, false, *FUniqueNetIdGooglePlay::EmptyId(), FString("Already trying to login"));
+		TriggerOnLoginCompleteDelegates(0, true, *UniqueNetId, TEXT(""));
+		return false;
 	}
+
+	auto Settings = GetDefault<UAndroidRuntimeSettings>();
+
+	MainSubsystem->QueueAsyncTask(new FOnlineAsyncTaskGooglePlayLogin(MainSubsystem, Settings->PlayGamesClientId, Settings->bForceRefreshToken));
 	
-	return bStartedLogin;
+	return true;
 }
 
+void FOnlineIdentityGooglePlay::SetIdentityData(FUniqueNetIdGooglePlayPtr InPlayerNetId, FString InPlayerAlias, FString InAuthCode)
+{
+	UniqueNetId = MoveTemp(InPlayerNetId);
+	PlayerAlias = MoveTemp(InPlayerAlias);
+	AuthCode = MoveTemp(InAuthCode);
+}
 
 bool FOnlineIdentityGooglePlay::Logout(int32 LocalUserNum)
 {
-	MainSubsystem->StartLogoutTask(LocalUserNum);
-	return true;
+	if (LocalUserNum == 0)
+	{
+		bool bWasLoggedIn = UniqueNetId->IsValid();
+		ClearIdentity();
+		if(bWasLoggedIn)
+		{
+			TriggerOnLoginStatusChangedDelegates(0, ELoginStatus::LoggedIn, ELoginStatus::NotLoggedIn, *UniqueNetId);
+			TriggerOnLoginChangedDelegates(0);
+		}
+	}
+	return false;
 }
 
-
 bool FOnlineIdentityGooglePlay::AutoLogin(int32 LocalUserNum)
 {
 	return Login(LocalUserNum, FOnlineAccountCredentials());
 }
 
-
 ELoginStatus::Type FOnlineIdentityGooglePlay::GetLoginStatus(int32 LocalUserNum) const
 {
-	if (LocalUserNum < MAX_LOCAL_PLAYERS && MainSubsystem->GetGameServices() != nullptr && MainSubsystem->GetGameServices()->IsAuthorized())
+	if (LocalUserNum == 0 && UniqueNetId->IsValid())
 	{
-		return ELoginStatus::LoggedIn;
+		return  ELoginStatus::LoggedIn;
+	}
+	else
+	{
+		return ELoginStatus::NotLoggedIn;
 	}
-
-	return ELoginStatus::NotLoggedIn;
 }
 
-
 ELoginStatus::Type FOnlineIdentityGooglePlay::GetLoginStatus(const FUniqueNetId& UserId) const
 {
-	if (MainSubsystem->GetGameServices() != nullptr && MainSubsystem->GetGameServices()->IsAuthorized())
+	if (UserId.IsValid() && UserId == *UniqueNetId)
 	{
 		return ELoginStatus::LoggedIn;
 	}
-
-	return ELoginStatus::NotLoggedIn;
+	else
+	{
+		return ELoginStatus::NotLoggedIn;
+	}
 }
 
-
 FUniqueNetIdPtr FOnlineIdentityGooglePlay::GetUniquePlayerId(int32 LocalUserNum) const
 {
-	if (UniqueNetId.IsValid())
+	if (LocalUserNum == 0)
 	{
 		return UniqueNetId;
 	}
@@ -109,7 +115,6 @@ FUniqueNetIdPtr FOnlineIdentityGooglePlay::GetUniquePlayerId(int32 LocalUserNum)
 	return FUniqueNetIdGooglePlay::EmptyId();
 }
 
-
 FUniqueNetIdPtr FOnlineIdentityGooglePlay::CreateUniquePlayerId(uint8* Bytes, int32 Size)
 {
 	if( Bytes && Size == sizeof(uint64) )
@@ -124,29 +129,45 @@ FUniqueNetIdPtr FOnlineIdentityGooglePlay::CreateUniquePlayerId(uint8* Bytes, in
 	return NULL;
 }
 
-
 FUniqueNetIdPtr FOnlineIdentityGooglePlay::CreateUniquePlayerId(const FString& Str)
 {
 	return FUniqueNetIdGooglePlay::Create(Str);
 }
 
-
 FString FOnlineIdentityGooglePlay::GetPlayerNickname(int32 LocalUserNum) const
 {
-	UE_LOG_ONLINE_IDENTITY(Display, TEXT("FOnlineIdentityGooglePlay::GetPlayerNickname"));
-	return PlayerAlias;
+	if (LocalUserNum == 0)
+	{ 
+		return PlayerAlias;
+	}
+	else
+	{
+		return FString();
+	}
 }
 
 FString FOnlineIdentityGooglePlay::GetPlayerNickname(const FUniqueNetId& UserId) const
 {
-	UE_LOG_ONLINE_IDENTITY(Display, TEXT("FOnlineIdentityGooglePlay::GetPlayerNickname"));
-	return PlayerAlias;
+	if (UserId.IsValid() && UserId == *UniqueNetId)
+	{
+		return GetPlayerNickname(0);
+	}
+	else
+	{
+		return FString();
+	}
 }
 
 FString FOnlineIdentityGooglePlay::GetAuthToken(int32 LocalUserNum) const
 {
-	UE_LOG_ONLINE_IDENTITY(Display, TEXT("FOnlineIdentityGooglePlay::GetAuthToken"));
-	return AuthToken;
+	if (LocalUserNum == 0)
+	{ 
+		return AuthCode;
+	}
+	else
+	{
+		return FString();
+	}
 }
 
 void FOnlineIdentityGooglePlay::RevokeAuthToken(const FUniqueNetId& UserId, const FOnRevokeAuthTokenCompleteDelegate& Delegate)
@@ -159,23 +180,6 @@ void FOnlineIdentityGooglePlay::RevokeAuthToken(const FUniqueNetId& UserId, cons
 	});
 }
 
-void FOnlineIdentityGooglePlay::Tick(float DeltaTime)
-{
-}
-
-void FOnlineIdentityGooglePlay::OnLoginCompleted(const int playerID, const gpg::AuthStatus errorCode)
-{
-	static const int32 MAX_TEXT_LINE_LEN = 32;
-	TCHAR Line[MAX_TEXT_LINE_LEN + 1] = { 0 };
-	int32 Len = FCString::Snprintf(Line, MAX_TEXT_LINE_LEN, TEXT("%d"), playerID);
-
-	UniqueNetId = FUniqueNetIdGooglePlay::Create(Line);
-	bLoggedIn = errorCode == gpg::AuthStatus::VALID;
-	TriggerOnLoginCompleteDelegates(playerID, bLoggedIn, *UniqueNetId, TEXT(""));
-
-	PendingConnectRequest.IsConnectionPending = false;
-}
-
 void FOnlineIdentityGooglePlay::GetUserPrivilege(const FUniqueNetId& UserId, EUserPrivileges::Type Privilege, const FOnGetUserPrivilegeCompleteDelegate& Delegate, EShowPrivilegeResolveUI ShowResolveUI)
 {
 	Delegate.ExecuteIfBound(UserId, Privilege, (uint32)EPrivilegeResults::NoFailures);
@@ -199,15 +203,3 @@ FString FOnlineIdentityGooglePlay::GetAuthType() const
 {
 	return TEXT("");
 }
-
-void FOnlineIdentityGooglePlay::SetPlayerDataFromFetchSelfResponse(const gpg::Player& PlayerData)
-{
-	const FString PlayerId(PlayerData.Id().c_str());
-	UniqueNetId = FUniqueNetIdGooglePlay::Create(PlayerId);
-	PlayerAlias = PlayerData.Name().c_str();
-}
-
-void FOnlineIdentityGooglePlay::SetAuthTokenFromGoogleConnectResponse(const FString& NewAuthToken)
-{
-	AuthToken = NewAuthToken;
-}
\ No newline at end of file
diff --git a/Engine/Plugins/Online/Android/OnlineSubsystemGooglePlay/Source/Private/OnlineIdentityInterfaceGooglePlay.h b/Engine/Plugins/Online/Android/OnlineSubsystemGooglePlay/Source/Private/OnlineIdentityInterfaceGooglePlay.h
--- a/Engine/Plugins/Online/Android/OnlineSubsystemGooglePlay/Source/Private/OnlineIdentityInterfaceGooglePlay.h
+++ b/Engine/Plugins/Online/Android/OnlineSubsystemGooglePlay/Source/Private/OnlineIdentityInterfaceGooglePlay.h
@@ -6,11 +6,6 @@
 #include "Interfaces/OnlineIdentityInterface.h"
 #include "OnlineSubsystemGooglePlayPackage.h"
 
-THIRD_PARTY_INCLUDES_START
-#include "gpg/status.h"
-#include "gpg/player.h"
-THIRD_PARTY_INCLUDES_END
-
 // from OnlineSubsystemTypes.h
 TEMP_UNIQUENETIDSTRING_SUBCLASS(FUniqueNetIdGooglePlay, GOOGLEPLAY_SUBSYSTEM);
 
@@ -18,46 +13,19 @@ class FOnlineIdentityGooglePlay :
 	public IOnlineIdentity
 {
 private:
-	
-	bool bPrevLoggedIn;
-	bool bLoggedIn;
 	FString PlayerAlias;
-	FString AuthToken;
-	int32 CurrentLoginUserNum;
-	class FOnlineSubsystemGooglePlay* MainSubsystem;
-
-	bool bLoggingInUser;
-	bool bRegisteringUser;
-
-	/** UID for this identity */
+	FString AuthCode;
 	FUniqueNetIdGooglePlayPtr UniqueNetId;
+	class FOnlineSubsystemGooglePlay* MainSubsystem;
 
-	struct FPendingConnection
-	{
-		FOnlineIdentityGooglePlay * ConnectionInterface;
-		bool IsConnectionPending;
-	};
-
-	/** Instance of the connection context data */
-	static FPendingConnection PendingConnectRequest;
+    static_assert(MAX_LOCAL_PLAYERS == 1, "FOnlineIdentityGooglePlay does not support more than 1 local player");
 
 PACKAGE_SCOPE:
 
-	/**
-	 * Default Constructor
-	 */
 	FOnlineIdentityGooglePlay(FOnlineSubsystemGooglePlay* InSubsystem);
 
-	/** Allow individual interfaces to access the currently signed-in user's id */
-	FUniqueNetIdGooglePlayPtr GetCurrentUserId() const { return UniqueNetId; }
-
-	void SetCurrentUserId(FUniqueNetIdGooglePlayPtr InUniqueNetId) { UniqueNetId = InUniqueNetId; }
-
-	/** Called from ExternalUIInterface to set UniqueId and PlayerAlias after authentication */
-	void SetPlayerDataFromFetchSelfResponse(const gpg::Player& PlayerData);
-
-	/** Called from ExternalUIInterface to set AuthToken after authentication */
-	void SetAuthTokenFromGoogleConnectResponse(const FString& NewAuthToken);
+	void SetIdentityData(FUniqueNetIdGooglePlayPtr PlayerNetId, FString PlayerAlias, FString AuthCode);
+	void ClearIdentity();
 
 public:
 
@@ -79,20 +47,7 @@ public:
 	virtual void GetUserPrivilege(const FUniqueNetId& UserId, EUserPrivileges::Type Privilege, const FOnGetUserPrivilegeCompleteDelegate& Delegate, EShowPrivilegeResolveUI ShowResolveUI=EShowPrivilegeResolveUI::Default) override;
 	virtual FPlatformUserId GetPlatformUserIdFromUniqueNetId(const FUniqueNetId& NetId) const override;
 	virtual FString GetAuthType() const override;
-	//~ End IOnlineIdentity interface
-
-public:
-
-	//~ Begin IOnlineIdentity Interface
-	void Tick(float DeltaTime);
-	//~ End IOnlineIdentity Interface
-	
-	//platform specific
-	void OnLogin(const bool InLoggedIn,  const FString& InPlayerId, const FString& InPlayerAlias);
-	void OnLogout(const bool InLoggedIn);
-
-	void OnLoginCompleted(const int playerID, gpg::AuthStatus errorCode);
-
+	//~ End IOnlineIdentity Interface	
 };
 
 
diff --git a/Engine/Plugins/Online/Android/OnlineSubsystemGooglePlay/Source/Private/OnlineLeaderboardGooglePlayCommon.h b/Engine/Plugins/Online/Android/OnlineSubsystemGooglePlay/Source/Private/OnlineLeaderboardGooglePlayCommon.h
new file mode 100644
--- /dev/null
+++ b/Engine/Plugins/Online/Android/OnlineSubsystemGooglePlay/Source/Private/OnlineLeaderboardGooglePlayCommon.h
@@ -0,0 +1,11 @@
+// Copyright Epic Games, Inc. All Rights Reserved.
+
+#pragma once
+
+struct FGooglePlayLeaderboardScore
+{
+	// Leaderboard ID as shown in GooglePlay Console
+	FString GooglePlayLeaderboardId;
+	// Leaderboard score
+	int64 Score = 0;
+};
diff --git a/Engine/Plugins/Online/Android/OnlineSubsystemGooglePlay/Source/Private/OnlineLeaderboardInterfaceGooglePlay.cpp b/Engine/Plugins/Online/Android/OnlineSubsystemGooglePlay/Source/Private/OnlineLeaderboardInterfaceGooglePlay.cpp
--- a/Engine/Plugins/Online/Android/OnlineSubsystemGooglePlay/Source/Private/OnlineLeaderboardInterfaceGooglePlay.cpp
+++ b/Engine/Plugins/Online/Android/OnlineSubsystemGooglePlay/Source/Private/OnlineLeaderboardInterfaceGooglePlay.cpp
@@ -1,13 +1,29 @@
 // Copyright Epic Games, Inc. All Rights Reserved.
 
 #include "OnlineLeaderboardInterfaceGooglePlay.h"
+#include "Algo/NoneOf.h"
+#include "AndroidRuntimeSettings.h"
+#include "OnlineAchievementGooglePlayCommon.h"
 #include "OnlineAsyncTaskGooglePlayReadLeaderboard.h"
-#include "UObject/Class.h"
+#include "OnlineAsyncTaskGooglePlayFlushLeaderboards.h"
 #include "OnlineSubsystemGooglePlay.h"
 
-THIRD_PARTY_INCLUDES_START
-#include "gpg/leaderboard_manager.h"
-THIRD_PARTY_INCLUDES_END
+namespace LeaderboardsGooglePlayDetail
+{
+	static TOptional<FString> GetGooglePlayLeaderboardId(const UAndroidRuntimeSettings* Settings, const FString& LeaderboardName)
+	{
+		for(const auto& Mapping : Settings->LeaderboardMap)
+		{
+			if(Mapping.Name == LeaderboardName)
+			{
+				return Mapping.LeaderboardID;
+			}
+		}
+
+		UE_LOG_ONLINE_LEADERBOARD(Warning, TEXT( "GetGooglePlayLeaderboardId: No mapping for leaderboard %s"), *LeaderboardName );
+		return NullOpt;
+	}
+}
 
 FOnlineLeaderboardsGooglePlay::FOnlineLeaderboardsGooglePlay(FOnlineSubsystemGooglePlay* InSubsystem)
 	: Subsystem(InSubsystem)
@@ -15,26 +31,34 @@ FOnlineLeaderboardsGooglePlay::FOnlineLeaderboardsGooglePlay(FOnlineSubsystemGoo
 	check(Subsystem);
 }
 
-bool FOnlineLeaderboardsGooglePlay::ReadLeaderboards(const TArray< FUniqueNetIdRef >& Players, FOnlineLeaderboardReadRef& ReadObject)
+bool FOnlineLeaderboardsGooglePlay::IsLocalPlayer(const FUniqueNetId& PlayerId) const
 {
-	ReadObject->Rows.Empty();
-
-	if (Subsystem->GetGameServices() == nullptr)
+	if (!PlayerId.IsValid())
 	{
-		ReadObject->ReadState = EOnlineAsyncTaskState::Failed;
-		Subsystem->GetLeaderboardsInterface()->TriggerOnLeaderboardReadCompleteDelegates(false);
 		return false;
 	}
+	FOnlineIdentityGooglePlayPtr IdentityInt = Subsystem->GetIdentityGooglePlay();
+	FUniqueNetIdPtr LocalPlayerNetId = IdentityInt->GetUniquePlayerId(0);
+	return LocalPlayerNetId? *LocalPlayerNetId == PlayerId : false;
+}
 
-	ReadObject->ReadState = EOnlineAsyncTaskState::InProgress;
+bool FOnlineLeaderboardsGooglePlay::ReadLeaderboards(const TArray< FUniqueNetIdRef >& Players, FOnlineLeaderboardReadRef& ReadObject)
+{
+	if (Algo::NoneOf(Players, [this](const FUniqueNetIdRef& PlayerId) { return IsLocalPlayer(*PlayerId); }))
+	{
+		UE_LOG_ONLINE_LEADERBOARD(Warning, TEXT("ReadLeaderboards failed because was called using non local players or local player is not logged in"));
+		return false;
+	}
 
-	auto ReadTask = new FOnlineAsyncTaskGooglePlayReadLeaderboard(
-		Subsystem,
-		ReadObject,
-		GetLeaderboardID(ReadObject->LeaderboardName.ToString()));
-	Subsystem->QueueAsyncTask(ReadTask);
+	ReadObject->Rows.Empty();
 
-	return true;
+	auto Settings = GetDefault<UAndroidRuntimeSettings>();
+	if(TOptional<FString> PlatformLeaderboardId = LeaderboardsGooglePlayDetail::GetGooglePlayLeaderboardId(Settings, ReadObject->LeaderboardName.ToString()))
+	{
+		Subsystem->QueueAsyncTask(new FOnlineAsyncTaskGooglePlayReadLeaderboard( Subsystem, ReadObject, *PlatformLeaderboardId));
+		return true;
+	}
+	return false;
 }
 
 bool FOnlineLeaderboardsGooglePlay::ReadLeaderboardsForFriends(int32 LocalUserNum, FOnlineLeaderboardReadRef& ReadObject)
@@ -51,52 +75,56 @@ void FOnlineLeaderboardsGooglePlay::FreeStats(FOnlineLeaderboardRead& ReadObject
 
 bool FOnlineLeaderboardsGooglePlay::WriteLeaderboards(const FName& SessionName, const FUniqueNetId& Player, FOnlineLeaderboardWrite& WriteObject)
 {
-	UE_LOG_ONLINE_LEADERBOARD(Display, TEXT("WriteLeaderboards"));
+	if (!IsLocalPlayer(Player))
+	{
+		UE_LOG_ONLINE_LEADERBOARD(Warning, TEXT("WriteLeaderboards failed because was called using non local player or local player is not logged in"));
+		return false;
+	}
 
 	bool bWroteAnyLeaderboard = false;
+	auto Settings = GetDefault<UAndroidRuntimeSettings>();
 
 	for(int32 LeaderboardIdx = 0; LeaderboardIdx < WriteObject.LeaderboardNames.Num(); ++LeaderboardIdx)
 	{
 		FString LeaderboardName = WriteObject.LeaderboardNames[LeaderboardIdx].ToString();
-		if(LeaderboardName.Equals(TEXT("TestLeaderboard")))
-		{
-			LeaderboardName = TEXT("leaderboard_00"); //hack to work around leaderboard name mismatch between test module and the format that java call expects
-		}
 		UE_LOG_ONLINE_LEADERBOARD(Display, TEXT("Going through stats for leaderboard :  %s "), *LeaderboardName);
 		
-		for(FStatPropertyArray::TConstIterator It(WriteObject.Properties); It; ++It)
+		for(auto &[Key, Stat]: WriteObject.Properties)
 		{
-			const FVariantData& Stat = It.Value();
-			uint64 Score;
-
-			UE_LOG_ONLINE_LEADERBOARD(Display, TEXT("Here's a stat"));
-
-			//Google leaderboard stats are always a long/int64
-			if(Stat.GetType() == EOnlineKeyValuePairDataType::Int64)
+			TOptional<FString> GooglePlayLeaderboardId = LeaderboardsGooglePlayDetail::GetGooglePlayLeaderboardId(Settings, LeaderboardName);
+			if (!GooglePlayLeaderboardId)
 			{
-				Stat.GetValue(Score);
-				
-				FOnlinePendingLeaderboardWrite* UnreportedScore = new (UnreportedScores) FOnlinePendingLeaderboardWrite();
-				UnreportedScore->LeaderboardName = LeaderboardName;
-				UnreportedScore->Score = Score;
-				UE_LOG_ONLINE_LEADERBOARD(Display, TEXT("FOnlineLeaderboardsGooglePlay::WriteLeaderboards() Int64 value Score: %d"), Score);
-
-				bWroteAnyLeaderboard = true;
+				continue;
 			}
-			else if (Stat.GetType() == EOnlineKeyValuePairDataType::Int32)
-			{
-				// cast from 32 to 64
-				int32 Score32 = 0;
-				Stat.GetValue(Score32);
-				Score = static_cast<uint64>(Score32);
 
-				FOnlinePendingLeaderboardWrite* UnreportedScore = new (UnreportedScores) FOnlinePendingLeaderboardWrite();
-				UnreportedScore->LeaderboardName = LeaderboardName;
-				UnreportedScore->Score = Score;
-				UE_LOG_ONLINE_LEADERBOARD(Display, TEXT("FOnlineLeaderboardsGooglePlay::WriteLeaderboards() Int32 value Score: %d "), Score);
+			int64 Score = 0;
 
-				bWroteAnyLeaderboard = true;
+			switch(Stat.GetType())
+			{
+				case EOnlineKeyValuePairDataType::Int64:
+				{
+					Stat.GetValue(Score);
+					break;
+				}
+				case EOnlineKeyValuePairDataType::Int32:
+				{
+					// cast from 32 to 64
+					int32 Score32 = 0;
+					Stat.GetValue(Score32);
+					Score = Score32;
+					break;
+				}
+				default:
+				{
+					UE_LOG_ONLINE_LEADERBOARD(Warning, TEXT("Unsupported stat type %s"), *LexToString(Stat.GetType()));
+					continue;
+				}
 			}
+			UE_LOG_ONLINE_LEADERBOARD(Display, TEXT("FOnlineLeaderboardsGooglePlay::WriteLeaderboards() %s value Score: %s"), *LexToString(Stat.GetType()), *Stat.ToString());
+			FGooglePlayLeaderboardScore& UnreportedScore = UnreportedScores.Emplace_GetRef();
+			UnreportedScore.GooglePlayLeaderboardId = MoveTemp(*GooglePlayLeaderboardId);
+			UnreportedScore.Score = Score;
+			bWroteAnyLeaderboard = true;
 		}
 	}
 	
@@ -106,32 +134,8 @@ bool FOnlineLeaderboardsGooglePlay::WriteLeaderboards(const FName& SessionName,
 
 bool FOnlineLeaderboardsGooglePlay::FlushLeaderboards(const FName& SessionName)
 {
-	UE_LOG_ONLINE_LEADERBOARD(Display, TEXT("flush leaderboards session name :%s"), *SessionName.ToString());
-
-	if (Subsystem->GetGameServices() == nullptr)
-	{
-		Subsystem->GetLeaderboardsInterface()->TriggerOnLeaderboardFlushCompleteDelegates(SessionName, false);
-		return false;
-	}
-
-	bool Success = true;
-
-	for(int32 Index = 0; Index < UnreportedScores.Num(); ++Index)
-	{
-		UE_LOG_ONLINE_LEADERBOARD(Display, TEXT("Submitting an unreported score to %s. Value: %d "), *UnreportedScores[Index].LeaderboardName);
-
-		const FString GoogleId = GetLeaderboardID(UnreportedScores[Index].LeaderboardName);
-
-		auto ConvertedId = FOnlineSubsystemGooglePlay::ConvertFStringToStdString(GoogleId);
-		Subsystem->GetGameServices()->Leaderboards().SubmitScore(
-			ConvertedId,
-			UnreportedScores[Index].Score);
-	}
-
+	Subsystem->QueueAsyncTask(new FOnlineAsyncTaskGooglePlayFlushLeaderboards( Subsystem, SessionName, MoveTemp(UnreportedScores)));
 	UnreportedScores.Empty();
-
-	TriggerOnLeaderboardFlushCompleteDelegates(SessionName, Success);
-
 	return true;
 }
 
@@ -150,19 +154,4 @@ bool FOnlineLeaderboardsGooglePlay::WriteOnlinePlayerRatings(const FName& Sessio
 {
 	//iOS doesn't support this, and there is no Google Play functionality for this either
 	return false;
-}
-
-FString FOnlineLeaderboardsGooglePlay::GetLeaderboardID(const FString& LeaderboardName)
-{
-	auto DefaultSettings = GetDefault<UAndroidRuntimeSettings>();
-	for(const auto& Mapping : DefaultSettings->LeaderboardMap)
-	{
-		if(Mapping.Name.Equals(LeaderboardName))
-		{
-			return Mapping.LeaderboardID;
-		}
-	}
-
-	UE_LOG_ONLINE_LEADERBOARD(Warning, TEXT( "GetLeaderboardID: No mapping for leaderboard %s"), *LeaderboardName );
-	return LeaderboardName;
-}
+}
\ No newline at end of file
diff --git a/Engine/Plugins/Online/Android/OnlineSubsystemGooglePlay/Source/Private/OnlineLeaderboardInterfaceGooglePlay.h b/Engine/Plugins/Online/Android/OnlineSubsystemGooglePlay/Source/Private/OnlineLeaderboardInterfaceGooglePlay.h
--- a/Engine/Plugins/Online/Android/OnlineSubsystemGooglePlay/Source/Private/OnlineLeaderboardInterfaceGooglePlay.h
+++ b/Engine/Plugins/Online/Android/OnlineSubsystemGooglePlay/Source/Private/OnlineLeaderboardInterfaceGooglePlay.h
@@ -4,13 +4,8 @@
 
 #include "CoreMinimal.h"
 #include "Interfaces/OnlineLeaderboardInterface.h"
-#include "AndroidRuntimeSettings.h"
-
-struct FOnlinePendingLeaderboardWrite
-{
-	FString LeaderboardName;
-	uint64 Score;
-};
+#include "OnlineLeaderboardGooglePlayCommon.h"
+#include "OnlineSubsystemGooglePlayPackage.h"
 
 /**
  * Interface definition for the online services leaderboard services 
@@ -33,16 +28,13 @@ private:
 	/** Pointer to owning subsystem */
 	FOnlineSubsystemGooglePlay* Subsystem;
 
-	/**
-	 * Helper function to get the platform- and game-specific leaderboard ID from the JSON config file.
-	 *
-	 * @param LeaderboardName the cross-platform name of the leaderboard to look up
-	 * @return The unique ID for the leaderboard as specified in the config file.
-	*/
 	FString GetLeaderboardID(const FString& LeaderboardName);
 
 	/** Scores are cached here in WriteLeaderboards until FlushLeaderboards is called */
-	TArray<FOnlinePendingLeaderboardWrite> UnreportedScores;
+	TArray<FGooglePlayLeaderboardScore> UnreportedScores;
+
+	/** Asks the identity interface it this PlayerId refers to the local player. We can only get information about the local player from GooglePlay */
+	bool IsLocalPlayer(const FUniqueNetId& PlayerId) const;
 };
 
 typedef TSharedPtr<FOnlineLeaderboardsGooglePlay, ESPMode::ThreadSafe> FOnlineLeaderboardsGooglePlayPtr;
diff --git a/Engine/Plugins/Online/Android/OnlineSubsystemGooglePlay/Source/Private/OnlineSubsystemGooglePlay.cpp b/Engine/Plugins/Online/Android/OnlineSubsystemGooglePlay/Source/Private/OnlineSubsystemGooglePlay.cpp
--- a/Engine/Plugins/Online/Android/OnlineSubsystemGooglePlay/Source/Private/OnlineSubsystemGooglePlay.cpp
+++ b/Engine/Plugins/Online/Android/OnlineSubsystemGooglePlay/Source/Private/OnlineSubsystemGooglePlay.cpp
@@ -11,35 +11,16 @@
 #include "OnlineExternalUIInterfaceGooglePlay.h"
 #include "OnlinePurchaseGooglePlay.h"
 #include "OnlineStoreGooglePlay.h"
-#include "Android/AndroidApplication.h"
-#include "Android/AndroidJNI.h"
 #include "OnlineAsyncTaskManagerGooglePlay.h"
-#include "OnlineAsyncTaskGooglePlayLogin.h"
-#include "OnlineAsyncTaskGooglePlayLogout.h"
-#include "OnlineAsyncTaskGooglePlayShowLoginUI.h"
 #include "Misc/ConfigCacheIni.h"
 #include "Async/TaskGraphInterfaces.h"
 #include "Stats/Stats.h"
 
-THIRD_PARTY_INCLUDES_START
-#include <android_native_app_glue.h>
-
-#include "gpg/android_initialization.h"
-#include "gpg/builder.h"
-#include "gpg/debug.h"
-#include "gpg/android_support.h"
-THIRD_PARTY_INCLUDES_END
-
-using namespace gpg;
-
 FOnlineSubsystemGooglePlay::FOnlineSubsystemGooglePlay(FName InInstanceName)
 	: FOnlineSubsystemImpl(GOOGLEPLAY_SUBSYSTEM, InInstanceName)
 	, IdentityInterface(nullptr)
 	, LeaderboardsInterface(nullptr)
 	, AchievementsInterface(nullptr)
-	, CurrentLoginTask(nullptr)
-	, CurrentShowLoginUITask(nullptr)
-	, CurrentLogoutTask(nullptr)
 {
 }
 
@@ -123,9 +104,6 @@ IOnlineAchievementsPtr FOnlineSubsystemGooglePlay::GetAchievementsInterface() co
 	return AchievementsInterface;
 }
 
-static bool WaitForLostFocus = false;
-static bool WaitingForLogin = false;
-
 bool FOnlineSubsystemGooglePlay::Init() 
 {
 	FPlatformMisc::LowLevelOutputDebugStringf(TEXT("FOnlineSubsystemGooglePlay::Init"));
@@ -133,6 +111,8 @@ bool FOnlineSubsystemGooglePlay::Init()
 	OnlineAsyncTaskThreadRunnable.Reset(new FOnlineAsyncTaskManagerGooglePlay);
 	OnlineAsyncTaskThread.Reset(FRunnableThread::Create(OnlineAsyncTaskThreadRunnable.Get(), *FString::Printf(TEXT("OnlineAsyncTaskThread %s"), *InstanceName.ToString())));
 
+	GooglePlayGamesWrapper.Init();
+
 	IdentityInterface = MakeShareable(new FOnlineIdentityGooglePlay(this));
 	LeaderboardsInterface = MakeShareable(new FOnlineLeaderboardsGooglePlay(this));
 	AchievementsInterface = MakeShareable(new FOnlineAchievementsGooglePlay(this));
@@ -145,15 +125,6 @@ bool FOnlineSubsystemGooglePlay::Init()
 		PurchaseInterface->Init();
 	}
 	
-	extern struct android_app* GNativeAndroidApp;
-	check(GNativeAndroidApp != nullptr);
-	AndroidInitialization::android_main(GNativeAndroidApp);
-
-	extern jobject GJavaGlobalNativeActivity;
-	PlatformConfiguration.SetActivity(GNativeAndroidApp->activity->clazz);
-
-	OnActivityResultDelegateHandle = FJavaWrapper::OnActivityResultDelegate.AddRaw(this, &FOnlineSubsystemGooglePlay::OnActivityResult);
-
 	return true;
 }
 
@@ -181,8 +152,6 @@ bool FOnlineSubsystemGooglePlay::Shutdown()
 
 	FOnlineSubsystemImpl::Shutdown();
 
-	FJavaWrapper::OnActivityResultDelegate.Remove(OnActivityResultDelegateHandle);
-
 #define DESTRUCT_INTERFACE(Interface) \
 	if (Interface.IsValid()) \
 	{ \
@@ -199,6 +168,8 @@ bool FOnlineSubsystemGooglePlay::Shutdown()
 	DESTRUCT_INTERFACE(IdentityInterface);
 #undef DESTRUCT_INTERFACE
 
+	GooglePlayGamesWrapper.Reset();
+
 	OnlineAsyncTaskThread.Reset();
 	OnlineAsyncTaskThreadRunnable.Reset();
 
@@ -231,7 +202,7 @@ bool FOnlineSubsystemGooglePlay::IsEnabled() const
 {
 	bool bEnabled = false;
 
-	// GameCircleRuntimeSettings holds a value for editor ease of use
+	// AndroidRuntimeSettings holds a value for editor ease of use
 	if (!GConfig->GetBool(TEXT("/Script/AndroidRuntimeSettings.AndroidRuntimeSettings"), TEXT("bEnableGooglePlaySupport"), bEnabled, GEngineIni))
 	{
 		UE_LOG_ONLINE(Warning, TEXT("The [/Script/AndroidRuntimeSettings.AndroidRuntimeSettings]:bEnableGooglePlaySupport flag has not been set"));
@@ -251,156 +222,8 @@ bool FOnlineSubsystemGooglePlay::IsInAppPurchasingEnabled()
 	return bSupportsInAppPurchasing;
 }
 
-void FOnlineSubsystemGooglePlay::StartShowLoginUITask(int PlayerId, const FOnLoginUIClosedDelegate& Delegate)
-{
-	UE_LOG_ONLINE(Log, TEXT("StartShowLoginUITask PlayerId: %d"), PlayerId);
-
-	if (AreAnyAsyncLoginTasksRunning())
-	{
-		UE_LOG_ONLINE(Log, TEXT("FOnlineSubsystemGooglePlay::StartShowLoginUITask: An asynchronous login task is already running."));
-		Delegate.ExecuteIfBound(nullptr, PlayerId, OnlineIdentity::Errors::LoginPending());
-		return;
-	}
-
-	if (GameServicesPtr.get() == nullptr)
-	{
-		UE_LOG_ONLINE(Log, TEXT("StartShowLoginUITask Game Services was null"));
-		// This is likely the first login attempt during this run. Attempt to create the
-		// GameServices object, which will automatically start a "silent" login attempt.
-		// If that succeeds, there's no need to show the login UI explicitly. If it fails,
-		// we'll call ShowAuthorizationUI.
-		
-		auto TheDelegate = FOnlineAsyncTaskGooglePlayLogin::FOnCompletedDelegate::CreateLambda([this, PlayerId, Delegate]()
-		{
-			UE_LOG_ONLINE(Log, TEXT("StartShowLoginUITask starting ShowLoginUITask_Internal"));
-			 StartShowLoginUITask_Internal(PlayerId, Delegate);
-		});
-
-		CurrentLoginTask = new FOnlineAsyncTaskGooglePlayLogin(this, PlayerId, TheDelegate);
-		QueueAsyncTask(CurrentLoginTask);
-	}
-	else
-	{
-		UE_LOG_ONLINE(Log, TEXT("StartShowLoginUITask GameServicesPtr valid"));
-		// We already have a GameServices object, so we can directly go to ShowAuthorizationUI.
-		StartShowLoginUITask_Internal(PlayerId, Delegate);
-	}
-}
-
-void FOnlineSubsystemGooglePlay::StartLogoutTask(int32 LocalUserNum)
-{
-	if (CurrentLogoutTask != nullptr)
-	{
-		UE_LOG_ONLINE(Log, TEXT("FOnlineSubsystemGooglePlay::StartLogoutTask: A logout task is already in progress."));
-		IdentityInterface->TriggerOnLogoutCompleteDelegates(LocalUserNum, false);
-		return;
-	}
-
-	CurrentLogoutTask = new FOnlineAsyncTaskGooglePlayLogout(this, LocalUserNum);
-	QueueAsyncTask(CurrentLogoutTask);
-}
-
-void FOnlineSubsystemGooglePlay::ProcessGoogleClientConnectResult(bool bInSuccessful, FString AccessToken)
-{
-	if (CurrentShowLoginUITask != nullptr)
-	{
-		// Only one login task should be active at a time
-		check(CurrentLoginTask == nullptr);
-
-		CurrentShowLoginUITask->ProcessGoogleClientConnectResult(bInSuccessful, AccessToken);
-	}
-}
-
-void FOnlineSubsystemGooglePlay::StartShowLoginUITask_Internal(int PlayerId, const FOnLoginUIClosedDelegate& Delegate)
-{
-	check(!AreAnyAsyncLoginTasksRunning());
-
-	UE_LOG_ONLINE(Log, TEXT("StartShowLoginUITask_Internal"));
-	CurrentShowLoginUITask = new FOnlineAsyncTaskGooglePlayShowLoginUI(this, PlayerId, Delegate);
-	QueueAsyncTask(CurrentShowLoginUITask);
-}
-
 void FOnlineSubsystemGooglePlay::QueueAsyncTask(FOnlineAsyncTask* AsyncTask)
 {
 	check(OnlineAsyncTaskThreadRunnable);
 	OnlineAsyncTaskThreadRunnable->AddToInQueue(AsyncTask);
-}
-
-void FOnlineSubsystemGooglePlay::OnAuthActionFinished(AuthOperation Op, AuthStatus Status)
-{
-	if (Op == AuthOperation::SIGN_IN)
-	{
-		UE_LOG_ONLINE(Log, TEXT("OnAuthActionFinished SIGN IN %d"), (int32)Status);
-		if (CurrentLoginTask != nullptr)
-		{
-			// Only one login task should be active at a time
-			check(CurrentShowLoginUITask == nullptr);
-
-			CurrentLoginTask->OnAuthActionFinished(Op, Status);
-		}
-		else if(CurrentShowLoginUITask != nullptr)
-		{
-			// Only one login task should be active at a time
-			check(CurrentLoginTask == nullptr);
-
-			CurrentShowLoginUITask->OnAuthActionFinished(Op, Status);
-		}
-		else
-		{
-			UE_LOG_ONLINE(Log, TEXT("OnAuthActionFinished no handler!"));
-		}
-	}
-	else if (Op == AuthOperation::SIGN_OUT)
-	{
-		UE_LOG_ONLINE(Log, TEXT("OnAuthActionFinished SIGN OUT %d"), (int32)Status);
-		if (CurrentLogoutTask != nullptr)
-		{
-			CurrentLogoutTask->OnAuthActionFinished(Op, Status);
-		}
-	}
-}
-
-std::string FOnlineSubsystemGooglePlay::ConvertFStringToStdString(const FString& InString)
-{
-	int32 SrcLen  = InString.Len() + 1;
-	int32 DestLen = FPlatformString::ConvertedLength<ANSICHAR>(*InString, SrcLen);
-	TArray<ANSICHAR> Converted;
-	Converted.AddUninitialized(DestLen);
-	
-	FPlatformString::Convert(Converted.GetData(), DestLen, *InString, SrcLen);
-
-	return std::string(Converted.GetData());
-}
-
-void FOnlineSubsystemGooglePlay::OnActivityResult(JNIEnv *env, jobject thiz, jobject activity, jint requestCode, jint resultCode, jobject data)
-{
-	// Pass the result on to google play - otherwise, some callbacks for the turn based system do not get called.
-	AndroidSupport::OnActivityResult(env, activity, requestCode, resultCode, data);
-}
-
-JNI_METHOD void Java_com_epicgames_unreal_GameActivity_nativeGoogleClientConnectCompleted(JNIEnv* jenv, jobject thiz, jboolean bSuccess, jstring accessToken)
-{
-	FString AccessToken;
-	if (bSuccess)
-	{
-		AccessToken = FJavaHelper::FStringFromParam(jenv, accessToken);
-	}
-
-	UE_LOG_ONLINE(Log, TEXT("nativeGoogleClientConnectCompleted Success: %d Token: %s"), bSuccess, *AccessToken);
-
-	DECLARE_CYCLE_STAT(TEXT("FSimpleDelegateGraphTask.ProcessGoogleClientConnectResult"), STAT_FSimpleDelegateGraphTask_ProcessGoogleClientConnectResult, STATGROUP_TaskGraphTasks);
-
-	FSimpleDelegateGraphTask::CreateAndDispatchWhenReady(
-		FSimpleDelegateGraphTask::FDelegate::CreateLambda([=]()
-		{
-			FPlatformMisc::LowLevelOutputDebugStringf(TEXT("Google Client connected %s, Access Token: %s\n"), bSuccess ? TEXT("successfully") : TEXT("unsuccessfully"), *AccessToken);
-			if (FOnlineSubsystemGooglePlay* const OnlineSub = static_cast<FOnlineSubsystemGooglePlay* const>(IOnlineSubsystem::Get(GOOGLEPLAY_SUBSYSTEM)))
-			{
-				OnlineSub->ProcessGoogleClientConnectResult(bSuccess, AccessToken);
-			}
-		}),
-		GET_STATID(STAT_FSimpleDelegateGraphTask_ProcessGoogleClientConnectResult),
-		nullptr,
-		ENamedThreads::GameThread
-	);
-}
+}
\ No newline at end of file
diff --git a/Engine/Plugins/Online/Android/OnlineSubsystemGooglePlay/Source/Public/OnlineSubsystemGooglePlay.h b/Engine/Plugins/Online/Android/OnlineSubsystemGooglePlay/Source/Public/OnlineSubsystemGooglePlay.h
--- a/Engine/Plugins/Online/Android/OnlineSubsystemGooglePlay/Source/Public/OnlineSubsystemGooglePlay.h
+++ b/Engine/Plugins/Online/Android/OnlineSubsystemGooglePlay/Source/Public/OnlineSubsystemGooglePlay.h
@@ -3,6 +3,7 @@
 #pragma once
 
 #include "CoreMinimal.h"
+#include "GooglePlayGamesWrapper.h"
 #include "OnlineSubsystem.h"
 #include "OnlineSubsystemImpl.h"
 #include "OnlineExternalUIInterfaceGooglePlay.h"
@@ -12,13 +13,6 @@
 #include "OnlineStoreGooglePlay.h"
 #include "OnlinePurchaseGooglePlay.h"
 
-THIRD_PARTY_INCLUDES_START
-#include "gpg/game_services.h"
-#include "gpg/android_platform_configuration.h"
-THIRD_PARTY_INCLUDES_END
-
-#include <string>
-
 /** Forward declarations of all interface classes */
 typedef TSharedPtr<class FOnlineIdentityGooglePlay,  ESPMode::ThreadSafe> FOnlineIdentityGooglePlayPtr;
 typedef TSharedPtr<class FOnlineStoreGooglePlayV2, ESPMode::ThreadSafe> FOnlineStoreGooglePlayV2Ptr;
@@ -27,9 +21,6 @@ typedef TSharedPtr<class FOnlineLeaderboardsGooglePlay, ESPMode::ThreadSafe> FOn
 typedef TSharedPtr<class FOnlineAchievementsGooglePlay, ESPMode::ThreadSafe> FOnlineAchievementsGooglePlayPtr;
 typedef TSharedPtr<class FOnlineExternalUIGooglePlay, ESPMode::ThreadSafe> FOnlineExternalUIGooglePlayPtr;
 
-class FOnlineAsyncTaskGooglePlayLogin;
-class FOnlineAsyncTaskGooglePlayShowLoginUI;
-class FOnlineAsyncTaskGooglePlayLogout;
 class FRunnableThread;
 
 /**
@@ -100,15 +91,14 @@ PACKAGE_SCOPE:
 	 */
 	void QueueAsyncTask(class FOnlineAsyncTask* AsyncTask);
 
-	/** Returns a pointer to the Google API entry point */
-	gpg::GameServices* GetGameServices() const { return GameServicesPtr.get(); }
-
-	/** Utility function, useful for Google APIs that take a std::string but we only have an FString */
-	static std::string ConvertFStringToStdString(const FString& InString);
+	FGooglePlayGamesWrapper& GetGooglePlayGamesWrapper() { return GooglePlayGamesWrapper; }
 
 	/** Returns the Google Play-specific version of Identity, useful to avoid unnecessary casting */
 	FOnlineIdentityGooglePlayPtr GetIdentityGooglePlay() const { return IdentityInterface; }
 
+	/** Returns the Google Play-specific version of Leaderboards, useful to avoid unnecessary casting */
+	FOnlineLeaderboardsGooglePlayPtr GetLeaderboardsGooglePlay() const { return LeaderboardsInterface; }
+
 	/** Returns the Google Play-specific version of Achievements, useful to avoid unnecessary casting */
 	FOnlineAchievementsGooglePlayPtr GetAchievementsGooglePlay() const { return AchievementsInterface; }
 
@@ -118,35 +108,16 @@ PACKAGE_SCOPE:
 	 */
 	bool IsInAppPurchasingEnabled();
 
-	/** Returns true if there are any async login tasks currently being tracked. */
-	bool AreAnyAsyncLoginTasksRunning() const { return CurrentLoginTask != nullptr || CurrentShowLoginUITask != nullptr; }
-
-	/** Start a ShowLoginUI async task. Creates the GameServices object first if necessary. */
-	void StartShowLoginUITask(int PlayerId, const FOnLoginUIClosedDelegate& Delegate = FOnLoginUIClosedDelegate());
-
-	/** Start a logout task if one isn't already in progress. */
-	void StartLogoutTask(int32 LocalUserNum);
-
-	/** Callback from JNI when Google Client is connected */
-	void ProcessGoogleClientConnectResult(bool bInSuccessful, FString AccessToken);
-
 private:
 
-	/** Google callback when auth is complete */
-	void OnAuthActionFinished(gpg::AuthOperation Op, gpg::AuthStatus Status);
-
-	/** Android callback when an activity is finished */
-	void OnActivityResult(JNIEnv *env, jobject thiz, jobject activity, jint requestCode, jint resultCode, jobject data);
-
-	/** Start a ShowLoginUI async task. */
-	void StartShowLoginUITask_Internal(int PlayerId, const FOnLoginUIClosedDelegate& Delegate = FOnLoginUIClosedDelegate());
-
 	/** Online async task runnable */
 	TUniquePtr<class FOnlineAsyncTaskManagerGooglePlay> OnlineAsyncTaskThreadRunnable;
 
 	/** Online async task thread */
 	TUniquePtr<FRunnableThread> OnlineAsyncTaskThread;
 
+	FGooglePlayGamesWrapper GooglePlayGamesWrapper;
+
 	/** Interface to the online identity system */
 	FOnlineIdentityGooglePlayPtr IdentityInterface;
 
@@ -164,45 +135,6 @@ private:
 
 	/** Interface to the external UI services */
 	FOnlineExternalUIGooglePlayPtr ExternalUIInterface;
-
-	/** Pointer to the main entry point for the Google API */
-	std::unique_ptr<gpg::GameServices> GameServicesPtr;
-
-	// Since the OnAuthActionFinished callback is "global" (within the scope of one GameServices object),
-	// we track the async tasks that depend on this callback to notify us that they're finished.
-	// Only one of these pointers should be non-null at any given time.
-
-	/**
-	 * Track the current login task (if any) so callbacks can notify it.
-	 * Still owned by the AsyncTaskManager, do not delete via this pointer!
-	 **/
-	FOnlineAsyncTaskGooglePlayLogin* CurrentLoginTask;
-
-	/**
-	 * Track the current ShowLoginUI task (if any)
-	 * Still owned by the AsyncTaskManager, do not delete via this pointer!
-	 **/
-	FOnlineAsyncTaskGooglePlayShowLoginUI* CurrentShowLoginUITask;
-
-	/**
-	 * Track the current Logout task
-	 * Still owned by the AsyncTaskManager, do not delete via this pointer!
-	 */
-	FOnlineAsyncTaskGooglePlayLogout* CurrentLogoutTask;
-
-	/** Handle of registered delegate for onActivityResult */
-	FDelegateHandle OnActivityResultDelegateHandle;
-
-	/** This task needs to be able to set the GameServicesPtr and clear CurrentLoginTask*/
-	friend class FOnlineAsyncTaskGooglePlayLogin;
-
-	/** This task needs to be able to clear CurrentShowLoginUITask */
-	friend class FOnlineAsyncTaskGooglePlayShowLoginUI;
-
-	/** This task needs to be able to clear CurrentLogoutTask */
-	friend class FOnlineAsyncTaskGooglePlayLogout;
-
-	gpg::AndroidPlatformConfiguration PlatformConfiguration;
 };
 
 typedef TSharedPtr<FOnlineSubsystemGooglePlay, ESPMode::ThreadSafe> FOnlineSubsystemGooglePlayPtr;
diff --git a/Engine/Source/Programs/UnrealBuildTool/Platform/Android/AndroidToolChain.cs b/Engine/Source/Programs/UnrealBuildTool/Platform/Android/AndroidToolChain.cs
--- a/Engine/Source/Programs/UnrealBuildTool/Platform/Android/AndroidToolChain.cs
+++ b/Engine/Source/Programs/UnrealBuildTool/Platform/Android/AndroidToolChain.cs
@@ -77,7 +77,7 @@ namespace UnrealBuildTool
 
 		private static Dictionary<UnrealArch, string[]> LibrariesToSkip = new() {
 			{ UnrealArch.Arm64, new string[] { "nvToolsExt", "nvToolsExtStub", "vorbisenc", } },
-			{ UnrealArch.X64,   new string[] { "nvToolsExt", "nvToolsExtStub", "oculus", "OVRPlugin", "vrapi", "ovrkernel", "systemutils", "openglloader", "ovrplatformloader", "gpg", "vorbisenc", } }
+			{ UnrealArch.X64,   new string[] { "nvToolsExt", "nvToolsExtStub", "oculus", "OVRPlugin", "vrapi", "ovrkernel", "systemutils", "openglloader", "ovrplatformloader", "vorbisenc", } }
 		};
 
 		private static Dictionary<UnrealArch, string[]> ModulesToSkip = new() {
diff --git a/Engine/Source/Programs/UnrealBuildTool/Platform/Android/UEDeployAndroid.cs b/Engine/Source/Programs/UnrealBuildTool/Platform/Android/UEDeployAndroid.cs
--- a/Engine/Source/Programs/UnrealBuildTool/Platform/Android/UEDeployAndroid.cs
+++ b/Engine/Source/Programs/UnrealBuildTool/Platform/Android/UEDeployAndroid.cs
@@ -2588,10 +2588,6 @@ namespace UnrealBuildTool
 			ConfigHierarchy Ini = GetConfigCacheIni(ConfigHierarchyType.Engine);
 			string PackageName = GetPackageName(ProjectName);
 			string VersionDisplayName = GetVersionDisplayName(bIsEmbedded);
-			bool bEnableGooglePlaySupport;
-			Ini.GetBool("/Script/AndroidRuntimeSettings.AndroidRuntimeSettings", "bEnableGooglePlaySupport", out bEnableGooglePlaySupport);
-			bool bUseGetAccounts;
-			Ini.GetBool("/Script/AndroidRuntimeSettings.AndroidRuntimeSettings", "bUseGetAccounts", out bUseGetAccounts);
 			string DepthBufferPreference;
 			Ini.GetString("/Script/AndroidRuntimeSettings.AndroidRuntimeSettings", "DepthBufferPreference", out DepthBufferPreference);
 			float MaxAspectRatioValue;
@@ -2930,10 +2926,6 @@ namespace UnrealBuildTool
 				{
 					StartupPermissions = StartupPermissions + (StartupPermissions.Length > 0 ? "," : "") + "android.permission.WRITE_EXTERNAL_STORAGE";
 				}
-				if (bEnableGooglePlaySupport && bUseGetAccounts)
-				{
-					StartupPermissions = StartupPermissions + (StartupPermissions.Length > 0 ? "," : "") + "android.permission.GET_ACCOUNTS";
-				}
 			}
 
 			Text.AppendLine(String.Format("\t\t<meta-data android:name=\"com.epicgames.unreal.GameActivity.EngineVersion\" android:value=\"{0}\"/>", EngineVersion));
@@ -3072,11 +3064,6 @@ namespace UnrealBuildTool
 					Text.AppendLine("\t<uses-permission android:name=\"android.permission.RECEIVE_BOOT_COMPLETED\"/>");
 				}
 
-				if (bEnableGooglePlaySupport && bUseGetAccounts)
-				{
-					Text.AppendLine("\t<uses-permission android:name=\"android.permission.GET_ACCOUNTS\"/>");
-				}
-
 				if(!bPackageForMetaQuest)
 				{
 					Text.AppendLine("\t<uses-permission android:name=\"android.permission.MODIFY_AUDIO_SETTINGS\"/>");
diff --git a/Engine/Source/Runtime/Android/AndroidRuntimeSettings/Classes/AndroidRuntimeSettings.h b/Engine/Source/Runtime/Android/AndroidRuntimeSettings/Classes/AndroidRuntimeSettings.h
--- a/Engine/Source/Runtime/Android/AndroidRuntimeSettings/Classes/AndroidRuntimeSettings.h
+++ b/Engine/Source/Runtime/Android/AndroidRuntimeSettings/Classes/AndroidRuntimeSettings.h
@@ -477,9 +477,14 @@ public:
 	UPROPERTY(GlobalConfig, EditAnywhere, Category = GooglePlayServices)
 	bool bEnableGooglePlaySupport;
 
-	// Enabling this adds GET_ACCOUNTS to manifest and user must give permission.  Required for reset achievements.
-	UPROPERTY(GlobalConfig, EditAnywhere, Category = GooglePlayServices, meta = (DisplayName = "Request Access Token On Connect"))
-	bool bUseGetAccounts;
+	UPROPERTY(GlobalConfig, EditAnywhere, Category = GooglePlayServices, meta = (DisplayName = "Request code for showing GooglePlay Games activities"), meta = (ToolTip = "This request code is reported in onActivityResult when Achievements or Leaderboards UIs are closed"))
+	int RequestCodeForPlayGamesActivities;
+
+	UPROPERTY(GlobalConfig, EditAnywhere, Category = GooglePlayServices, meta = (DisplayName = "GooglePlay Games Oauth Client ID registered for game server"), meta = (ToolTip = "No auth code will be requested on login if empty. Setting up Server-side API Access (see https://developers.google.com/games/services/android/offline-access)"))
+	FString PlayGamesClientId;
+
+	UPROPERTY(GlobalConfig, EditAnywhere, Category = GooglePlayServices, meta = (DisplayName = "Enable GooglePlay Games offline access"), meta = (ToolTip = "Include refresh token when Auth code is exchanged. Setting up Server-side API Access (see https://developers.google.com/games/services/android/offline-access)"))
+	bool bForceRefreshToken;
 
 	// The app id obtained from the Google Play Developer Console
 	UPROPERTY(GlobalConfig, EditAnywhere, Category = GooglePlayServices)
@@ -493,10 +498,6 @@ public:
 	UPROPERTY(GlobalConfig, EditAnywhere, Category = GooglePlayServices)
 	TArray<FGooglePlayLeaderboardMapping> LeaderboardMap;
 
-	// Enabling this requests snapshots support for saved games during Google Play login.
-	UPROPERTY(GlobalConfig, EditAnywhere, Category = GooglePlayServices, meta = (DisplayName = "Enable Snapshots on Google Play login [Experimental]"))
-	bool bEnableSnapshots;
-
 	// Enabling this includes the AdMob SDK and will be detected by Google Play Console on upload of APK.  Disable if you do not need ads to remove this warning.
 	UPROPERTY(GlobalConfig, EditAnywhere, Category = GooglePlayServices, meta = (DisplayName = "Include AdMob support for ads"))
 	bool bSupportAdMob;
diff --git a/Engine/Source/Runtime/Android/AndroidRuntimeSettings/Private/AndroidRuntimeSettings.cpp b/Engine/Source/Runtime/Android/AndroidRuntimeSettings/Private/AndroidRuntimeSettings.cpp
--- a/Engine/Source/Runtime/Android/AndroidRuntimeSettings/Private/AndroidRuntimeSettings.cpp
+++ b/Engine/Source/Runtime/Android/AndroidRuntimeSettings/Private/AndroidRuntimeSettings.cpp
@@ -24,7 +24,8 @@ UAndroidRuntimeSettings::UAndroidRuntimeSettings(const FObjectInitializer& Objec
 	, bAndroidVoiceEnabled(false)
 	, bPackageForMetaQuest(false)
 	, bEnableGooglePlaySupport(false)
-	, bUseGetAccounts(false)
+	, RequestCodeForPlayGamesActivities(80002)
+	, bForceRefreshToken(false)
 	, bSupportAdMob(true)
 	, bBlockAndroidKeysOnControllers(false)
 	, AudioSampleRate(44100)
diff --git a/Engine/Source/Runtime/Core/Private/Android/AndroidJavaEnv.cpp b/Engine/Source/Runtime/Core/Private/Android/AndroidJavaEnv.cpp
--- a/Engine/Source/Runtime/Core/Private/Android/AndroidJavaEnv.cpp
+++ b/Engine/Source/Runtime/Core/Private/Android/AndroidJavaEnv.cpp
@@ -191,10 +191,23 @@ jclass AndroidJavaEnv::FindJavaClass(const char* name)
 	{
 		return nullptr;
 	}
+
 	jstring ClassNameObj = Env->NewStringUTF(name);
 	jclass FoundClass = static_cast<jclass>(Env->CallObjectMethod(ClassLoader, FindClassMethod, ClassNameObj));
-	CheckJavaException();
+
+	if (!FoundClass)
+	{
+		// Clear exception because we failed to find the class and try again using JNIEnv
+		if (Env->ExceptionCheck())
+		{
+			Env->ExceptionClear();
+		}
+
+		FoundClass = static_cast<jclass>(Env->FindClass(name));
+		CheckJavaException();
+	}
 	Env->DeleteLocalRef(ClassNameObj);
+	
 	return FoundClass;
 }
 
@@ -205,9 +218,23 @@ jclass AndroidJavaEnv::FindJavaClassGlobalRef(const char* name)
 	{
 		return nullptr;
 	}
-	auto ClassNameObj = FJavaHelper::ToJavaString(Env, FString(ANSI_TO_TCHAR(name)));
-	auto FoundClass = NewScopedJavaObject(Env, static_cast<jclass>(Env->CallObjectMethod(ClassLoader, FindClassMethod, *ClassNameObj)));
-	CheckJavaException();
+
+	jstring ClassNameObj = Env->NewStringUTF(name);
+	auto FoundClass = NewScopedJavaObject(Env, static_cast<jclass>(Env->CallObjectMethod(ClassLoader, FindClassMethod, ClassNameObj)));
+	
+	if (!FoundClass)
+	{
+		// Clear exception because we failed to find the class and try again using JNIEnv
+		if (Env->ExceptionCheck())
+		{
+			Env->ExceptionClear();
+		}
+
+		FoundClass = NewScopedJavaObject(Env, static_cast<jclass>(Env->FindClass(name)));
+		CheckJavaException();
+	}
+	Env->DeleteLocalRef(ClassNameObj);
+
 	auto GlobalClass = (jclass)Env->NewGlobalRef(*FoundClass);
 	return GlobalClass;
 }
diff --git a/Engine/Source/Runtime/Launch/Private/Android/AndroidJNI.cpp b/Engine/Source/Runtime/Launch/Private/Android/AndroidJNI.cpp
--- a/Engine/Source/Runtime/Launch/Private/Android/AndroidJNI.cpp
+++ b/Engine/Source/Runtime/Launch/Private/Android/AndroidJNI.cpp
@@ -207,7 +207,6 @@ void FJavaWrapper::FindGooglePlayMethods(JNIEnv* Env)
 	// @todo split GooglePlay
 	//	GoogleServicesClassID = FindClass(Env, "com/epicgames/unreal/GoogleServices", bIsOptional);
 	GoogleServicesClassID = GameActivityClassID;
-	AndroidThunkJava_ResetAchievements = FindMethod(Env, GoogleServicesClassID, "AndroidThunkJava_ResetAchievements", "()V", bIsOptional);
 	AndroidThunkJava_ShowAdBanner = FindMethod(Env, GoogleServicesClassID, "AndroidThunkJava_ShowAdBanner", "(Ljava/lang/String;Z)V", bIsOptional);
 	AndroidThunkJava_HideAdBanner = FindMethod(Env, GoogleServicesClassID, "AndroidThunkJava_HideAdBanner", "()V", bIsOptional);
 	AndroidThunkJava_CloseAdBanner = FindMethod(Env, GoogleServicesClassID, "AndroidThunkJava_CloseAdBanner", "()V", bIsOptional);
@@ -216,8 +215,6 @@ void FJavaWrapper::FindGooglePlayMethods(JNIEnv* Env)
 	AndroidThunkJava_IsInterstitialAdRequested = FindMethod(Env, GoogleServicesClassID, "AndroidThunkJava_IsInterstitialAdRequested", "()Z", bIsOptional);
 	AndroidThunkJava_ShowInterstitialAd = FindMethod(Env, GoogleServicesClassID, "AndroidThunkJava_ShowInterstitialAd", "()V", bIsOptional);
 	AndroidThunkJava_GetAdvertisingId = FindMethod(Env, GoogleServicesClassID, "AndroidThunkJava_GetAdvertisingId", "()Ljava/lang/String;", bIsOptional);
-	AndroidThunkJava_GoogleClientConnect = FindMethod(Env, GoogleServicesClassID, "AndroidThunkJava_GoogleClientConnect", "()V", bIsOptional);
-	AndroidThunkJava_GoogleClientDisconnect = FindMethod(Env, GoogleServicesClassID, "AndroidThunkJava_GoogleClientDisconnect", "()V", bIsOptional);
 }
 void FJavaWrapper::FindGooglePlayBillingMethods(JNIEnv* Env)
 {
@@ -382,6 +379,109 @@ bool FJavaWrapper::CallBooleanMethod(JNIEnv* Env, jobject Object, jmethodID Meth
 	return (bool)Return;
 }
 
+void FJavaWrapper::CallStaticVoidMethod(JNIEnv* Env, jclass Clazz, jmethodID Method, ...)
+{
+	if (Method == NULL || Clazz == NULL)
+	{
+		return;
+	}
+
+	va_list Args;
+	va_start(Args, Method);
+	Env->CallStaticVoidMethodV(Clazz, Method, Args);
+	va_end(Args);
+}
+
+jobject FJavaWrapper::CallStaticObjectMethod(JNIEnv* Env, jclass Clazz, jmethodID Method, ...)
+{
+	if (Method == NULL || Clazz == NULL)
+	{
+		return nullptr;
+	}
+
+	va_list Args;
+	va_start(Args, Method);
+	jobject Return = Env->CallStaticObjectMethodV(Clazz, Method, Args);
+	va_end(Args);
+
+	return Return;
+}
+
+int32 FJavaWrapper::CallStaticIntMethod(JNIEnv* Env, jclass Clazz, jmethodID Method, ...)
+{
+	if (Method == NULL || Clazz == NULL)
+	{
+		return 0;
+	}
+
+	va_list Args;
+	va_start(Args, Method);
+	jint Return = Env->CallStaticIntMethod(Clazz, Method, Args);
+	va_end(Args);
+
+	return (int32)Return;
+}
+
+int64 FJavaWrapper::CallStaticLongMethod(JNIEnv* Env, jclass Clazz, jmethodID Method, ...)
+{
+	if (Method == NULL || Clazz == NULL)
+	{
+		return 0;
+	}
+
+	va_list Args;
+	va_start(Args, Method);
+	jlong Return = Env->CallStaticLongMethod(Clazz, Method, Args);
+	va_end(Args);
+
+	return (int64)Return;
+}
+
+float FJavaWrapper::CallStaticFloatMethod(JNIEnv* Env, jclass Clazz, jmethodID Method, ...)
+{
+	if (Method == NULL || Clazz == NULL)
+	{
+		return 0.f;
+	}
+
+	va_list Args;
+	va_start(Args, Method);
+	jfloat Return = Env->CallStaticFloatMethod(Clazz, Method, Args);
+	va_end(Args);
+
+	return (float)Return;
+}
+
+double FJavaWrapper::CallStaticDoubleMethod(JNIEnv* Env, jclass Clazz, jmethodID Method, ...)
+{
+	if (Method == NULL || Clazz == NULL)
+	{
+		return 0.;
+	}
+
+	va_list Args;
+	va_start(Args, Method);
+	jdouble Return = Env->CallStaticDoubleMethod(Clazz, Method, Args);
+	va_end(Args);
+
+	return (double)Return;
+}
+
+bool FJavaWrapper::CallStaticBooleanMethod(JNIEnv* Env, jclass Clazz, jmethodID Method, ...)
+{
+	if (Method == NULL || Clazz == NULL)
+	{
+		return false;
+	}
+
+	va_list Args;
+	va_start(Args, Method);
+	jboolean Return = Env->CallStaticBooleanMethod(Clazz, Method, Args);
+	va_end(Args);
+
+	return (bool)Return;
+}
+
 //Declare all the static members of the class defs 
 jclass FJavaWrapper::GameActivityClassID;
 jobject FJavaWrapper::GameActivityThis;
@@ -452,7 +552,6 @@ jfieldID FJavaWrapper::InputDeviceInfo_FeedbackMotorCount;
 
 jclass FJavaWrapper::GoogleServicesClassID;
 jobject FJavaWrapper::GoogleServicesThis;
-jmethodID FJavaWrapper::AndroidThunkJava_ResetAchievements;
 jmethodID FJavaWrapper::AndroidThunkJava_ShowAdBanner;
 jmethodID FJavaWrapper::AndroidThunkJava_HideAdBanner;
 jmethodID FJavaWrapper::AndroidThunkJava_CloseAdBanner;
@@ -461,8 +560,6 @@ jmethodID FJavaWrapper::AndroidThunkJava_IsInterstitialAdAvailable;
 jmethodID FJavaWrapper::AndroidThunkJava_IsInterstitialAdRequested;
 jmethodID FJavaWrapper::AndroidThunkJava_ShowInterstitialAd;
 jmethodID FJavaWrapper::AndroidThunkJava_GetAdvertisingId;
-jmethodID FJavaWrapper::AndroidThunkJava_GoogleClientConnect;
-jmethodID FJavaWrapper::AndroidThunkJava_GoogleClientDisconnect;
 
 jclass FJavaWrapper::JavaStringClass;
 jmethodID FJavaWrapper::AndroidThunkJava_IapSetupService;
@@ -1133,14 +1230,6 @@ void AndroidThunkCpp_LaunchURL(const FString& URL)
 	}
 }
 
-void AndroidThunkCpp_ResetAchievements()
-{
-	if (JNIEnv* Env = FAndroidApplication::GetJavaEnv())
-	{
-		FJavaWrapper::CallVoidMethod(Env, FJavaWrapper::GoogleServicesThis, FJavaWrapper::AndroidThunkJava_ResetAchievements);
-	}
-}
-
 void AndroidThunkCpp_ShowAdBanner(const FString& AdUnitID, bool bShowOnBottomOfScreen)
 {
 	if (JNIEnv* Env = FAndroidApplication::GetJavaEnv())
@@ -1276,22 +1365,6 @@ bool AndroidThunkCpp_SendBroadcast(const FString& PackageName, const FString& Ex
 	return result;
 }
 
-void AndroidThunkCpp_GoogleClientConnect()
-{
-	if (JNIEnv* Env = FAndroidApplication::GetJavaEnv())
-	{
-		FJavaWrapper::CallVoidMethod(Env, FJavaWrapper::GoogleServicesThis, FJavaWrapper::AndroidThunkJava_GoogleClientConnect);
-	}
-}
-
-void AndroidThunkCpp_GoogleClientDisconnect()
-{
-	if (JNIEnv* Env = FAndroidApplication::GetJavaEnv())
-	{
-		FJavaWrapper::CallVoidMethod(Env, FJavaWrapper::GoogleServicesThis, FJavaWrapper::AndroidThunkJava_GoogleClientDisconnect);
-	}
-}
-
 namespace
 {
 	jobject GJavaAssetManager = NULL;
diff --git a/Engine/Source/Runtime/Launch/Public/Android/AndroidJNI.h b/Engine/Source/Runtime/Launch/Public/Android/AndroidJNI.h
--- a/Engine/Source/Runtime/Launch/Public/Android/AndroidJNI.h
+++ b/Engine/Source/Runtime/Launch/Public/Android/AndroidJNI.h
@@ -97,7 +97,6 @@ public:
 	// IDs related to google play services
 	static jclass GoogleServicesClassID;
 	static jobject GoogleServicesThis;
-	static jmethodID AndroidThunkJava_ResetAchievements;
 	static jmethodID AndroidThunkJava_ShowAdBanner;
 	static jmethodID AndroidThunkJava_HideAdBanner;
 	static jmethodID AndroidThunkJava_CloseAdBanner;
@@ -106,9 +105,7 @@ public:
 	static jmethodID AndroidThunkJava_IsInterstitialAdRequested;
 	static jmethodID AndroidThunkJava_ShowInterstitialAd;
 	static jmethodID AndroidThunkJava_GetAdvertisingId;
-	static jmethodID AndroidThunkJava_GoogleClientConnect;
-	static jmethodID AndroidThunkJava_GoogleClientDisconnect;
-
+	
 	// Optionally added if GCM plugin (or other remote notification system) enabled
 	static jmethodID AndroidThunkJava_RegisterForRemoteNotifications;
 	static jmethodID AndroidThunkJava_UnregisterForRemoteNotifications;
@@ -176,10 +173,17 @@ public:
 	static int32 CallIntMethod(JNIEnv* Env, jobject Object, jmethodID Method, ...);
 	static int64 CallLongMethod(JNIEnv* Env, jobject Object, jmethodID Method, ...);
 	static float CallFloatMethod(JNIEnv* Env, jobject Object, jmethodID Method, ...);
-	static double CallDoubleMethod(JNIEnv* Env, jobject Object, jmethodID Method, ...);
-	
+	static double CallDoubleMethod(JNIEnv* Env, jobject Object, jmethodID Method, ...);	
 	static bool CallBooleanMethod(JNIEnv* Env, jobject Object, jmethodID Method, ...);
 
+	static void CallStaticVoidMethod(JNIEnv* Env, jclass Clazz, jmethodID Method, ...);
+	static jobject CallStaticObjectMethod(JNIEnv* Env, jclass Clazz, jmethodID Method, ...);
+	static int32 CallStaticIntMethod(JNIEnv* Env, jclass Clazz, jmethodID Method, ...);
+	static int64 CallStaticLongMethod(JNIEnv* Env, jclass Clazz, jmethodID Method, ...);
+	static float CallStaticFloatMethod(JNIEnv* Env, jclass Clazz, jmethodID Method, ...);
+	static double CallStaticDoubleMethod(JNIEnv* Env, jclass Clazz, jmethodID Method, ...);
+	static bool CallStaticBooleanMethod(JNIEnv* Env, jclass Clazz, jmethodID Method, ...);
+
 	// Delegate that can be registered to that is called when an activity is finished
 	static FOnActivityResult OnActivityResultDelegate;
 
