Expose server auth code in OSSGooglePlay through FUserOnlineAccount::GetAuthAttribute using  AUTH_ATTR_AUTHORIZATION_CODE as key

diff --git a/Engine/Plugins/Online/Android/OnlineSubsystemGooglePlay/Source/OnlineSubsystemGooglePlay.Build.cs b/Engine/Plugins/Online/Android/OnlineSubsystemGooglePlay/Source/OnlineSubsystemGooglePlay.Build.cs
--- a/Engine/Plugins/Online/Android/OnlineSubsystemGooglePlay/Source/OnlineSubsystemGooglePlay.Build.cs
+++ b/Engine/Plugins/Online/Android/OnlineSubsystemGooglePlay/Source/OnlineSubsystemGooglePlay.Build.cs
@@ -22,23 +22,15 @@ public class OnlineSubsystemGooglePlay : ModuleRules
 
 		PrivateDependencyModuleNames.AddRange(
 			new string[] { 
-                "Core", 
-                "CoreUObject", 
-                "Engine", 
-                "Sockets",
+				"Core",
+				"Engine",
+				"CoreOnline",
 				"OnlineSubsystem", 
-                "HTTP",
 				"AndroidRuntimeSettings",
 				"Launch"
             }
 			);
 
-		PublicIncludePathModuleNames.AddRange(
-			new string[] {
-				"AndroidPermission"
-			}
-			);
-
         string PluginPath = Utils.MakePathRelativeTo(ModuleDirectory, Target.RelativeEnginePath);
         AdditionalPropertiesForReceipt.Add("AndroidPlugin", Path.Combine(PluginPath, "OnlineSubsystemGooglePlay_UPL.xml"));
     }
diff --git a/Engine/Plugins/Online/Android/OnlineSubsystemGooglePlay/Source/Private/OnlineIdentityInterfaceGooglePlay.cpp b/Engine/Plugins/Online/Android/OnlineSubsystemGooglePlay/Source/Private/OnlineIdentityInterfaceGooglePlay.cpp
--- a/Engine/Plugins/Online/Android/OnlineSubsystemGooglePlay/Source/Private/OnlineIdentityInterfaceGooglePlay.cpp
+++ b/Engine/Plugins/Online/Android/OnlineSubsystemGooglePlay/Source/Private/OnlineIdentityInterfaceGooglePlay.cpp
@@ -5,6 +5,82 @@
 #include "OnlineAsyncTaskGooglePlayLogin.h"
 #include "OnlineSubsystemGooglePlay.h"
 
+class FUserOnlineAccountGooglePlay : public FUserOnlineAccount
+{
+public:
+	FUserOnlineAccountGooglePlay(const FUniqueNetIdGooglePlayRef& InPlayerNetId, FString InPlayerAlias, FString InAuthCode)
+		: UniqueNetId(InPlayerNetId)
+		, PlayerAlias(MoveTemp(InPlayerAlias))
+		, AuthCode(MoveTemp(InAuthCode))
+	{}
+
+	// FOnlineUser
+	virtual FUniqueNetIdRef GetUserId() const override 
+	{ 
+		return UniqueNetId; 
+	}
+	
+	virtual FString GetRealName() const override 
+	{ 
+		return FString(); 
+	}
+	
+	virtual FString GetDisplayName(const FString& Platform = FString()) const override 
+	{ 
+		return PlayerAlias; 
+	}
+
+	virtual bool GetUserAttribute(const FString& AttrName, FString& OutAttrValue) const override 
+	{ 
+		if (AttrName == USER_ATTR_DISPLAYNAME)
+		{
+			OutAttrValue = PlayerAlias;
+			return true;
+		}
+		else if (AttrName == USER_ATTR_ID)
+		{
+			OutAttrValue = UniqueNetId->ToString();
+            return true;
+		}
+		return false;
+	}
+	
+	virtual bool SetUserLocalAttribute(const FString& AttrName, const FString& InAttrValue) override 
+	{ 
+		return false;
+	}
+
+	// FUserOnlineAccount
+	virtual FString GetAccessToken() const override 
+	{ 
+		return FString(); 
+	}
+	
+	virtual bool HasAccessTokenExpired(const FDateTime& Time) const override 
+	{ 
+		return true; 
+	}
+
+	virtual bool GetAuthAttribute(const FString& AttrName, FString& OutAttrValue) const override
+	{
+		if (AttrName == AUTH_ATTR_AUTHORIZATION_CODE && !AuthCode.IsEmpty())
+		{
+			OutAttrValue = AuthCode;
+			return true;
+		}
+		return false;
+	}
+	
+	virtual bool SetUserAttribute(const FString& AttrName, const FString& AttrValue) override 
+	{ 
+		return false; 
+	}
+private:
+	FUniqueNetIdGooglePlayRef UniqueNetId;
+	FString PlayerAlias;
+	FString AuthCode;
+};
+
 FOnlineIdentityGooglePlay::FOnlineIdentityGooglePlay(FOnlineSubsystemGooglePlay* InSubsystem)
 	: MainSubsystem(InSubsystem)
 {
@@ -16,19 +92,25 @@ FOnlineIdentityGooglePlay::FOnlineIdentityGooglePlay(FOnlineSubsystemGooglePlay*
 
 void FOnlineIdentityGooglePlay::ClearIdentity()
 {
-	UniqueNetId = FUniqueNetIdGooglePlay::EmptyId();
-	PlayerAlias.Empty();
-	AuthCode.Empty();
+	LocalPlayerAccount.Reset();
 }
 
 TSharedPtr<FUserOnlineAccount> FOnlineIdentityGooglePlay::GetUserAccount(const FUniqueNetId& UserId) const
 {
+	if (LocalPlayerAccount && *LocalPlayerAccount->GetUserId() == UserId)
+	{
+		return LocalPlayerAccount;
+	}
 	return nullptr;
 }
 
 TArray<TSharedPtr<FUserOnlineAccount> > FOnlineIdentityGooglePlay::GetAllUserAccounts() const
 {
 	TArray<TSharedPtr<FUserOnlineAccount> > Result;
+	if (LocalPlayerAccount)
+	{
+		Result.Add(LocalPlayerAccount);
+	}
 
 	return Result;
 }
@@ -41,9 +123,9 @@ bool FOnlineIdentityGooglePlay::Login(int32 LocalUserNum, const FOnlineAccountCr
 		return false;
 	}
 	
-	if (UniqueNetId->IsValid())
+	if (LocalPlayerAccount.IsValid())
 	{
-		TriggerOnLoginCompleteDelegates(0, true, *UniqueNetId, TEXT(""));
+		TriggerOnLoginCompleteDelegates(0, true, *LocalPlayerAccount->GetUserId(), TEXT(""));
 		return false;
 	}
 
@@ -54,22 +136,20 @@ bool FOnlineIdentityGooglePlay::Login(int32 LocalUserNum, const FOnlineAccountCr
 	return true;
 }
 
-void FOnlineIdentityGooglePlay::SetIdentityData(FUniqueNetIdGooglePlayPtr InPlayerNetId, FString InPlayerAlias, FString InAuthCode)
+void FOnlineIdentityGooglePlay::SetIdentityData(const FUniqueNetIdGooglePlayPtr& InPlayerNetId, FString InPlayerAlias, FString InAuthCode)
 {
-	UniqueNetId = MoveTemp(InPlayerNetId);
-	PlayerAlias = MoveTemp(InPlayerAlias);
-	AuthCode = MoveTemp(InAuthCode);
+	LocalPlayerAccount = MakeShared<FUserOnlineAccountGooglePlay>(InPlayerNetId.ToSharedRef(), MoveTemp(InPlayerAlias), MoveTemp(InAuthCode));
 }
 
 bool FOnlineIdentityGooglePlay::Logout(int32 LocalUserNum)
 {
 	if (LocalUserNum == 0)
 	{
-		bool bWasLoggedIn = UniqueNetId->IsValid();
+		bool bWasLoggedIn = LocalPlayerAccount.IsValid();
 		ClearIdentity();
 		if(bWasLoggedIn)
 		{
-			TriggerOnLoginStatusChangedDelegates(0, ELoginStatus::LoggedIn, ELoginStatus::NotLoggedIn, *UniqueNetId);
+			TriggerOnLoginStatusChangedDelegates(0, ELoginStatus::LoggedIn, ELoginStatus::NotLoggedIn, *FUniqueNetIdGooglePlay::EmptyId());
 			TriggerOnLoginChangedDelegates(0);
 		}
 	}
@@ -83,7 +163,7 @@ bool FOnlineIdentityGooglePlay::AutoLogin(int32 LocalUserNum)
 
 ELoginStatus::Type FOnlineIdentityGooglePlay::GetLoginStatus(int32 LocalUserNum) const
 {
-	if (LocalUserNum == 0 && UniqueNetId->IsValid())
+	if (LocalUserNum == 0 && LocalPlayerAccount.IsValid())
 	{
 		return  ELoginStatus::LoggedIn;
 	}
@@ -95,7 +175,7 @@ ELoginStatus::Type FOnlineIdentityGooglePlay::GetLoginStatus(int32 LocalUserNum)
 
 ELoginStatus::Type FOnlineIdentityGooglePlay::GetLoginStatus(const FUniqueNetId& UserId) const
 {
-	if (UserId.IsValid() && UserId == *UniqueNetId)
+	if (UserId.IsValid() && LocalPlayerAccount && *LocalPlayerAccount->GetUserId() == UserId)
 	{
 		return ELoginStatus::LoggedIn;
 	}
@@ -107,9 +187,9 @@ ELoginStatus::Type FOnlineIdentityGooglePlay::GetLoginStatus(const FUniqueNetId&
 
 FUniqueNetIdPtr FOnlineIdentityGooglePlay::GetUniquePlayerId(int32 LocalUserNum) const
 {
-	if (LocalUserNum == 0)
+	if (LocalUserNum == 0 && LocalPlayerAccount)
 	{
-		return UniqueNetId;
+		return LocalPlayerAccount->GetUserId();
 	}
 
 	return FUniqueNetIdGooglePlay::EmptyId();
@@ -136,9 +216,9 @@ FUniqueNetIdPtr FOnlineIdentityGooglePlay::CreateUniquePlayerId(const FString& S
 
 FString FOnlineIdentityGooglePlay::GetPlayerNickname(int32 LocalUserNum) const
 {
-	if (LocalUserNum == 0)
+	if (LocalUserNum == 0 && LocalPlayerAccount)
 	{ 
-		return PlayerAlias;
+		return LocalPlayerAccount->GetDisplayName();
 	}
 	else
 	{
@@ -148,7 +228,7 @@ FString FOnlineIdentityGooglePlay::GetPlayerNickname(int32 LocalUserNum) const
 
 FString FOnlineIdentityGooglePlay::GetPlayerNickname(const FUniqueNetId& UserId) const
 {
-	if (UserId.IsValid() && UserId == *UniqueNetId)
+	if (UserId.IsValid() && LocalPlayerAccount && *LocalPlayerAccount->GetUserId() == UserId)
 	{
 		return GetPlayerNickname(0);
 	}
@@ -160,14 +240,11 @@ FString FOnlineIdentityGooglePlay::GetPlayerNickname(const FUniqueNetId& UserId)
 
 FString FOnlineIdentityGooglePlay::GetAuthToken(int32 LocalUserNum) const
 {
-	if (LocalUserNum == 0)
-	{ 
-		return AuthCode;
-	}
-	else
+	if (LocalUserNum == 0 && LocalPlayerAccount)
 	{
-		return FString();
+		return LocalPlayerAccount->GetAccessToken();
 	}
+	return FString();
 }
 
 void FOnlineIdentityGooglePlay::RevokeAuthToken(const FUniqueNetId& UserId, const FOnRevokeAuthTokenCompleteDelegate& Delegate)
diff --git a/Engine/Plugins/Online/Android/OnlineSubsystemGooglePlay/Source/Private/OnlineIdentityInterfaceGooglePlay.h b/Engine/Plugins/Online/Android/OnlineSubsystemGooglePlay/Source/Private/OnlineIdentityInterfaceGooglePlay.h
--- a/Engine/Plugins/Online/Android/OnlineSubsystemGooglePlay/Source/Private/OnlineIdentityInterfaceGooglePlay.h
+++ b/Engine/Plugins/Online/Android/OnlineSubsystemGooglePlay/Source/Private/OnlineIdentityInterfaceGooglePlay.h
@@ -13,9 +13,8 @@ class FOnlineIdentityGooglePlay :
 	public IOnlineIdentity
 {
 private:
-	FString PlayerAlias;
-	FString AuthCode;
-	FUniqueNetIdGooglePlayPtr UniqueNetId;
+	TSharedPtr<class FUserOnlineAccountGooglePlay> LocalPlayerAccount;
+
 	class FOnlineSubsystemGooglePlay* MainSubsystem;
 
     static_assert(MAX_LOCAL_PLAYERS == 1, "FOnlineIdentityGooglePlay does not support more than 1 local player");
@@ -24,7 +23,7 @@ PACKAGE_SCOPE:
 
 	FOnlineIdentityGooglePlay(FOnlineSubsystemGooglePlay* InSubsystem);
 
-	void SetIdentityData(FUniqueNetIdGooglePlayPtr PlayerNetId, FString PlayerAlias, FString AuthCode);
+	void SetIdentityData(const FUniqueNetIdGooglePlayPtr& PlayerNetId, FString PlayerAlias, FString AuthCode);
 	void ClearIdentity();
 
 public:
