Support for stats deprecation extended to GooglePlay platform

diff --git a/Engine/Plugins/Online/Android/OnlineSubsystemGooglePlay/Source/Private/OnlineAchievementsInterfaceGooglePlay.cpp b/Engine/Plugins/Online/Android/OnlineSubsystemGooglePlay/Source/Private/OnlineAchievementsInterfaceGooglePlay.cpp
--- a/Engine/Plugins/Online/Android/OnlineSubsystemGooglePlay/Source/Private/OnlineAchievementsInterfaceGooglePlay.cpp
+++ b/Engine/Plugins/Online/Android/OnlineSubsystemGooglePlay/Source/Private/OnlineAchievementsInterfaceGooglePlay.cpp
@@ -88,11 +88,11 @@ void FOnlineAchievementsGooglePlay::FinishAchievementWrite(
 	for (auto& [Key, Stat] : WriteObject->Properties)
 	{
 		// Create an achievement object which should be reported to the server.
-		const FOnlineAchievementGooglePlay* Achievement = Achievements.FindByPredicate([Key = Key.ToString()](const FOnlineAchievement& Achievement) { return Achievement.Id == Key;} );
+		const FOnlineAchievementGooglePlay* Achievement = Achievements.FindByPredicate([Key](const FOnlineAchievement& Achievement) { return Achievement.Id == Key;} );
 		
 		if (!Achievement)
 		{
-            UE_LOG_ONLINE_ACHIEVEMENTS(Error, TEXT("Unknown achievement id: %s"), *Key.ToString());
+            UE_LOG_ONLINE_ACHIEVEMENTS(Error, TEXT("Unknown achievement id: %s"), *Key);
             bAllSucceeded = false;
 			continue;
 		}
@@ -100,7 +100,7 @@ void FOnlineAchievementsGooglePlay::FinishAchievementWrite(
 		TOptional<FString> GoogleAchievementId = GetGoogleAchievementIdFromUnrealAchievementId(Settings, Achievement->Id);
 		if (!GoogleAchievementId)
 		{
-            UE_LOG_ONLINE_ACHIEVEMENTS(Error, TEXT("No Google achievement Id for achievement id: %s"), *Key.ToString());
+            UE_LOG_ONLINE_ACHIEVEMENTS(Error, TEXT("No Google achievement Id for achievement id: %s"), *Key);
             bAllSucceeded = false;
 			continue;
 		}
diff --git a/Engine/Plugins/Online/Android/OnlineSubsystemGooglePlay/Source/Private/OnlineAsyncTaskGooglePlayReadLeaderboard.cpp b/Engine/Plugins/Online/Android/OnlineSubsystemGooglePlay/Source/Private/OnlineAsyncTaskGooglePlayReadLeaderboard.cpp
--- a/Engine/Plugins/Online/Android/OnlineSubsystemGooglePlay/Source/Private/OnlineAsyncTaskGooglePlayReadLeaderboard.cpp
+++ b/Engine/Plugins/Online/Android/OnlineSubsystemGooglePlay/Source/Private/OnlineAsyncTaskGooglePlayReadLeaderboard.cpp
@@ -68,7 +68,7 @@ void FOnlineAsyncTaskGooglePlayReadLeaderboard::AddScore(const FString& DisplayN
 
 				default:
 				{
-					UE_LOG_ONLINE(Warning, TEXT("Unsupported key value pair during retrieval from Google Play %s"), *ColumnMeta.ColumnName.ToString());
+					UE_LOG_ONLINE(Warning, TEXT("Unsupported key value pair during retrieval from Google Play %s"), *ColumnMeta.ColumnName);
 					break;
 				}
 			}
diff --git a/Engine/Plugins/Online/Android/OnlineSubsystemGooglePlay/Source/Private/OnlineLeaderboardInterfaceGooglePlay.cpp b/Engine/Plugins/Online/Android/OnlineSubsystemGooglePlay/Source/Private/OnlineLeaderboardInterfaceGooglePlay.cpp
--- a/Engine/Plugins/Online/Android/OnlineSubsystemGooglePlay/Source/Private/OnlineLeaderboardInterfaceGooglePlay.cpp
+++ b/Engine/Plugins/Online/Android/OnlineSubsystemGooglePlay/Source/Private/OnlineLeaderboardInterfaceGooglePlay.cpp
@@ -53,7 +53,7 @@ bool FOnlineLeaderboardsGooglePlay::ReadLeaderboards(const TArray< FUniqueNetIdR
 	ReadObject->Rows.Empty();
 
 	auto Settings = GetDefault<UAndroidRuntimeSettings>();
-	if(TOptional<FString> PlatformLeaderboardId = LeaderboardsGooglePlayDetail::GetGooglePlayLeaderboardId(Settings, ReadObject->LeaderboardName.ToString()))
+	if(TOptional<FString> PlatformLeaderboardId = LeaderboardsGooglePlayDetail::GetGooglePlayLeaderboardId(Settings, ReadObject->LeaderboardName))
 	{
 		Subsystem->QueueAsyncTask(new FOnlineAsyncTaskGooglePlayReadLeaderboard( Subsystem, ReadObject, *PlatformLeaderboardId));
 		return true;
@@ -86,7 +86,7 @@ bool FOnlineLeaderboardsGooglePlay::WriteLeaderboards(const FName& SessionName,
 
 	for(int32 LeaderboardIdx = 0; LeaderboardIdx < WriteObject.LeaderboardNames.Num(); ++LeaderboardIdx)
 	{
-		FString LeaderboardName = WriteObject.LeaderboardNames[LeaderboardIdx].ToString();
+		FString LeaderboardName = WriteObject.LeaderboardNames[LeaderboardIdx];
 		UE_LOG_ONLINE_LEADERBOARD(Display, TEXT("Going through stats for leaderboard :  %s "), *LeaderboardName);
 		
 		for(auto &[Key, Stat]: WriteObject.Properties)
