Fix for Horde issue 536799

diff --git a/Engine/Plugins/Online/Android/OnlineSubsystemGooglePlay/Source/Private/OnlineAchievementsInterfaceGooglePlay.cpp b/Engine/Plugins/Online/Android/OnlineSubsystemGooglePlay/Source/Private/OnlineAchievementsInterfaceGooglePlay.cpp
--- a/Engine/Plugins/Online/Android/OnlineSubsystemGooglePlay/Source/Private/OnlineAchievementsInterfaceGooglePlay.cpp
+++ b/Engine/Plugins/Online/Android/OnlineSubsystemGooglePlay/Source/Private/OnlineAchievementsInterfaceGooglePlay.cpp
@@ -88,7 +88,7 @@ void FOnlineAchievementsGooglePlay::FinishAchievementWrite(
 	for (auto& [Key, Stat] : WriteObject->Properties)
 	{
 		// Create an achievement object which should be reported to the server.
-		const FOnlineAchievementGooglePlay* Achievement = Achievements.FindByPredicate([Key](const FOnlineAchievement& Achievement) { return Achievement.Id == Key;} );
+		const FOnlineAchievementGooglePlay* Achievement = Achievements.FindByPredicate([Key = Key](const FOnlineAchievement& Achievement) { return Achievement.Id == Key;} );
 		
 		if (!Achievement)
 		{
