Fix wrong check on successful Achievements write task

diff --git a/Engine/Plugins/Online/Android/OnlineSubsystemGooglePlay/Source/Private/OnlineAsyncTaskGooglePlayWriteAchievements.cpp b/Engine/Plugins/Online/Android/OnlineSubsystemGooglePlay/Source/Private/OnlineAsyncTaskGooglePlayWriteAchievements.cpp
--- a/Engine/Plugins/Online/Android/OnlineSubsystemGooglePlay/Source/Private/OnlineAsyncTaskGooglePlayWriteAchievements.cpp
+++ b/Engine/Plugins/Online/Android/OnlineSubsystemGooglePlay/Source/Private/OnlineAsyncTaskGooglePlayWriteAchievements.cpp
@@ -29,8 +29,8 @@ void FOnlineAsyncTaskGooglePlayWriteAchievements::Tick()
 
 bool FOnlineAsyncTaskGooglePlayWriteAchievements::SetSucceeded(const TArray<FString>& SucceededIds)
 {
-	// Remove data of entried that did not succeed
-	return 0 != WriteAchievementsData.RemoveAll([&SucceededIds](const FGooglePlayAchievementWriteData& Entry)
+	// Remove data of entries that did not succeed. Task succeeded if no entry was removed
+	return 0 == WriteAchievementsData.RemoveAll([&SucceededIds](const FGooglePlayAchievementWriteData& Entry)
 		{
 			return Algo::NoneOf(SucceededIds, [&Entry](const FString& Id) { return Id == Entry.GooglePlayAchievementId;} );
 		});
