Update GooglePlay Billing library to 7.1.1 in OSSGooglePlay

diff --git a/Engine/Plugins/Online/Android/OnlineSubsystemGooglePlay/Source/Java/com/epicgames/unreal/GooglePlayStoreHelper.java b/Engine/Plugins/Online/Android/OnlineSubsystemGooglePlay/Source/Java/com/epicgames/unreal/GooglePlayStoreHelper.java
--- a/Engine/Plugins/Online/Android/OnlineSubsystemGooglePlay/Source/Java/com/epicgames/unreal/GooglePlayStoreHelper.java
+++ b/Engine/Plugins/Online/Android/OnlineSubsystemGooglePlay/Source/Java/com/epicgames/unreal/GooglePlayStoreHelper.java
@@ -15,6 +15,7 @@ import com.android.billingclient.api.AcknowledgePurchaseParams;
 import com.android.billingclient.api.AcknowledgePurchaseResponseListener;
 import com.android.billingclient.api.ConsumeParams;
 import com.android.billingclient.api.ConsumeResponseListener;
+import com.android.billingclient.api.PendingPurchasesParams;
 import com.android.billingclient.api.ProductDetails;
 import com.android.billingclient.api.ProductDetailsResponseListener;
 import com.android.billingclient.api.Purchase;
@@ -53,9 +54,13 @@ public class GooglePlayStoreHelper implements StoreHelper, PurchasesUpdatedListe
 
 		MainActivity = InGameActivity;
 
+		PendingPurchasesParams PendingPurchases = PendingPurchasesParams.newBuilder()
+				.enableOneTimeProducts()
+				.build();
+
 		Client = BillingClient.newBuilder(MainActivity)
 			.setListener(this)
-			.enablePendingPurchases()
+			.enablePendingPurchases(PendingPurchases)
 			.build();
 
 		Client.startConnection(new BillingClientStateListener() {
@@ -167,7 +172,7 @@ public class GooglePlayStoreHelper implements StoreHelper, PurchasesUpdatedListe
 
 				// We will be querying for in app products and subscriptions 
 				int AwaitingResponseCount = 2;
-				List<ProductDetails> AllProductDetails = new ArrayList<>();
+				final List<ProductDetails> AllProductDetails = new ArrayList<>();
 			};
 
 			Client.queryProductDetailsAsync(QueryProducts, AggregatedListener);
@@ -326,7 +331,7 @@ public class GooglePlayStoreHelper implements StoreHelper, PurchasesUpdatedListe
 			}
 			// We will be querying for in app products and subscriptions 
 			int AwaitingResponseCount = 2;
-			List<Purchase> AllProducts = new ArrayList<>();
+			final List<Purchase> AllProducts = new ArrayList<>();
 		};
 
 		Client.queryPurchasesAsync(QueryInApp, AggregatedListener);
diff --git a/Engine/Plugins/Online/Android/OnlineSubsystemGooglePlay/Source/OnlineSubsystemGooglePlay_UPL.xml b/Engine/Plugins/Online/Android/OnlineSubsystemGooglePlay/Source/OnlineSubsystemGooglePlay_UPL.xml
--- a/Engine/Plugins/Online/Android/OnlineSubsystemGooglePlay/Source/OnlineSubsystemGooglePlay_UPL.xml
+++ b/Engine/Plugins/Online/Android/OnlineSubsystemGooglePlay/Source/OnlineSubsystemGooglePlay_UPL.xml
@@ -28,7 +28,7 @@
 			<true>
 				<insert>
 			dependencies {
-			implementation 'com.android.billingclient:billing:6.2.1'
+			implementation 'com.android.billingclient:billing:7.1.1'
 			}
 				</insert>
 			</true>
