Update Google Billing library to 8.1.0

#Synced-CL 48077828

diff --git a/Engine/Plugins/Online/Android/OnlineSubsystemGooglePlay/Source/Java/com/epicgames/unreal/GooglePlayStoreHelper.java b/Engine/Plugins/Online/Android/OnlineSubsystemGooglePlay/Source/Java/com/epicgames/unreal/GooglePlayStoreHelper.java
--- a/Engine/Plugins/Online/Android/OnlineSubsystemGooglePlay/Source/Java/com/epicgames/unreal/GooglePlayStoreHelper.java
+++ b/Engine/Plugins/Online/Android/OnlineSubsystemGooglePlay/Source/Java/com/epicgames/unreal/GooglePlayStoreHelper.java
@@ -22,6 +22,7 @@ import com.android.billingclient.api.Purchase;
 import com.android.billingclient.api.PurchasesResponseListener;
 import com.android.billingclient.api.PurchasesUpdatedListener;
 import com.android.billingclient.api.QueryProductDetailsParams;
+import com.android.billingclient.api.QueryProductDetailsResult;
 import com.android.billingclient.api.QueryPurchasesParams;
 import com.android.vending.billing.util.Base64; //WMM should use different base64 here.
 
@@ -55,12 +56,13 @@ public class GooglePlayStoreHelper implements StoreHelper, PurchasesUpdatedListe
 		MainActivity = InGameActivity;
 
 		PendingPurchasesParams PendingPurchases = PendingPurchasesParams.newBuilder()
-				.enableOneTimeProducts()
-				.build();
+			.enableOneTimeProducts()
+			.build();
 
 		Client = BillingClient.newBuilder(MainActivity)
 			.setListener(this)
 			.enablePendingPurchases(PendingPurchases)
+			.enableAutoServiceReconnection()
 			.build();
 
 		Client.startConnection(new BillingClientStateListener() {
@@ -136,7 +138,7 @@ public class GooglePlayStoreHelper implements StoreHelper, PurchasesUpdatedListe
 				// called from any thread and all the asynchronous callbacks will be returned on the same thread" those callbacks may be called on different 
 				// threads if invoked from JNI.
 				@Override
-				public synchronized void onProductDetailsResponse(@NonNull BillingResult Result, @NonNull List<ProductDetails> ProductDetailsList)
+				public synchronized void onProductDetailsResponse(@NonNull BillingResult Result, @NonNull QueryProductDetailsResult ProductDetailsList)
 				{
 					if(AwaitingResponseCount <= 0)
 					{
@@ -145,9 +147,9 @@ public class GooglePlayStoreHelper implements StoreHelper, PurchasesUpdatedListe
 					int ResponseCode = Result.getResponseCode();
 					if(ResponseCode == BillingClient.BillingResponseCode.OK)
 					{
-						AllProductDetails.addAll(ProductDetailsList);
+						AllProductDetails.addAll(ProductDetailsList.getProductDetailsList());
 						AwaitingResponseCount -= 1;
-						Log.debug("[GooglePlayStoreHelper] - QueryInAppPurchases - Received new " + ProductDetailsList.size() + " products. Waiting for additional responses: " + (AwaitingResponseCount > 0));
+						Log.debug("[GooglePlayStoreHelper] - QueryInAppPurchases - Received new " + ProductDetailsList.getProductDetailsList().size() + " products. Waiting for additional responses: " + (AwaitingResponseCount > 0));
 						if (AwaitingResponseCount == 0)
 						{
 							Log.debug("[GooglePlayStoreHelper] - QueryInAppPurchases - ProductDetails:" + AllProductDetails.toString());
@@ -218,12 +220,12 @@ public class GooglePlayStoreHelper implements StoreHelper, PurchasesUpdatedListe
 		QueryProductDetailsParams Params = CreateQueryProductDetailsParams(ProductIds, CommonProductType);
 		Client.queryProductDetailsAsync(Params, new ProductDetailsResponseListener() {
 			@Override
-			public void onProductDetailsResponse(@NonNull BillingResult Result, @NonNull List<ProductDetails> ProductDetailsList)
+			public void onProductDetailsResponse(@NonNull BillingResult Result, @NonNull QueryProductDetailsResult ProductDetailsList)
 			{
 				if(Result.getResponseCode() == BillingClient.BillingResponseCode.OK)
 				{
-					Log.debug("[GooglePlayStoreHelper] - BeginPurchase - CreateBillingFlowParams with " + ProductDetailsList.stream().map(ProductDetails::getName).collect(Collectors.joining(",")));
-					final BillingFlowParams Params = CreateBillingFlowParams(ProductIds, ProductDetailsList, ObfuscatedAccountId);
+					Log.debug("[GooglePlayStoreHelper] - BeginPurchase - CreateBillingFlowParams with " + ProductDetailsList.getProductDetailsList().stream().map(ProductDetails::getName).collect(Collectors.joining(",")));
+					final BillingFlowParams Params = CreateBillingFlowParams(ProductIds, ProductDetailsList.getProductDetailsList(), ObfuscatedAccountId);
 					if(Params == null)
 					{
 						Log.debug("[GooglePlayStoreHelper] - BeginPurchase - Failed! Matching products may not exist in the store");
@@ -235,7 +237,7 @@ public class GooglePlayStoreHelper implements StoreHelper, PurchasesUpdatedListe
 					{
 						@Override
 						public void run() {
-							Log.debug("[GooglePlayStoreHelper] - BeginPurchase - launchBillingFlow " + ProductDetailsList.stream().map(ProductDetails::getName).collect(Collectors.joining(",")));
+							Log.debug("[GooglePlayStoreHelper] - BeginPurchase - launchBillingFlow " + ProductDetailsList.getProductDetailsList().stream().map(ProductDetails::getName).collect(Collectors.joining(",")));
 							BillingResult Result = Client.launchBillingFlow(MainActivity, Params);
 							int ResponseCode = Result.getResponseCode();
 							if (ResponseCode != BillingClient.BillingResponseCode.OK)
@@ -303,12 +305,12 @@ public class GooglePlayStoreHelper implements StoreHelper, PurchasesUpdatedListe
 							QueryProductDetailsParams QuerySubs = CreateQueryProductDetailsParams(StoreProductIds, BillingClient.ProductType.SUBS);
 							Client.queryProductDetailsAsync(QuerySubs, new ProductDetailsResponseListener() {
 								@Override
-								public void onProductDetailsResponse(@NonNull BillingResult Result, @NonNull List<ProductDetails> ProductDetailsList)
+								public void onProductDetailsResponse(@NonNull BillingResult Result, @NonNull QueryProductDetailsResult ProductDetailsList)
 								{
 									int ResponseCode = Result.getResponseCode();
 									if (ResponseCode == BillingClient.BillingResponseCode.OK)
 									{
-										Set<String> SubscriptionProductIds = GetSubscriptionStoreProductIds(ProductDetailsList);
+										Set<String> SubscriptionProductIds = GetSubscriptionStoreProductIds(ProductDetailsList.getProductDetailsList());
 										ListOfPurchasesAsArrays PurchasesArrays = new ListOfPurchasesAsArrays(AllProducts, SubscriptionProductIds);
 										NativeQueryExistingPurchasesComplete(ResponseCode, PurchasesArrays.OwnedProducts, PurchasesArrays.PurchaseStates, PurchasesArrays.ProductTokens, PurchasesArrays.Receipts, PurchasesArrays.Signatures);
 									}
@@ -396,12 +398,12 @@ public class GooglePlayStoreHelper implements StoreHelper, PurchasesUpdatedListe
 			QueryProductDetailsParams QuerySubs = CreateQueryProductDetailsParams(StoreProductIds, BillingClient.ProductType.SUBS);
 			Client.queryProductDetailsAsync(QuerySubs, new ProductDetailsResponseListener() {
 				@Override
-				public void onProductDetailsResponse(@NonNull BillingResult Result, @NonNull List<ProductDetails> ProductDetailsList)
+				public void onProductDetailsResponse(@NonNull BillingResult Result, @NonNull QueryProductDetailsResult ProductDetailsList)
 				{
 					int responseCode = Result.getResponseCode();
 					if (responseCode == BillingClient.BillingResponseCode.OK)
 					{
-						Set<String> SubscriptionProductIds = GetSubscriptionStoreProductIds(ProductDetailsList);
+						Set<String> SubscriptionProductIds = GetSubscriptionStoreProductIds(ProductDetailsList.getProductDetailsList());
 
 						for (Purchase UpdatedPurchase : UpdatedPurchases)
 						{
diff --git a/Engine/Plugins/Online/Android/OnlineSubsystemGooglePlay/Source/OnlineSubsystemGooglePlay_UPL.xml b/Engine/Plugins/Online/Android/OnlineSubsystemGooglePlay/Source/OnlineSubsystemGooglePlay_UPL.xml
--- a/Engine/Plugins/Online/Android/OnlineSubsystemGooglePlay/Source/OnlineSubsystemGooglePlay_UPL.xml
+++ b/Engine/Plugins/Online/Android/OnlineSubsystemGooglePlay/Source/OnlineSubsystemGooglePlay_UPL.xml
@@ -34,7 +34,7 @@
 					<true>
 						<insert>
 			dependencies {
-			implementation 'com.android.billingclient:billing:7.1.1'
+			implementation 'com.android.billingclient:billing:8.1.0'
 			}
 						</insert>
 					</true>
