Updated the comments on IsAllowedToPurchase API to clarify that this API is meant for platform-specific policy or permission checks.

Always return true for the IsAllowedToPurchase API on the Google Play purchase interface. This is because that there are no policy or permission reasons for a Google Play purchase to be blocked.

diff --git a/Engine/Plugins/Online/Android/OnlineSubsystemGooglePlay/Source/Private/OnlinePurchaseGooglePlay.cpp b/Engine/Plugins/Online/Android/OnlineSubsystemGooglePlay/Source/Private/OnlinePurchaseGooglePlay.cpp
--- a/Engine/Plugins/Online/Android/OnlineSubsystemGooglePlay/Source/Private/OnlinePurchaseGooglePlay.cpp
+++ b/Engine/Plugins/Online/Android/OnlineSubsystemGooglePlay/Source/Private/OnlinePurchaseGooglePlay.cpp
@@ -221,16 +221,21 @@ bool FOnlinePurchaseGooglePlay::IsSubscriptionProductId(const FString& ProductId
 
 bool FOnlinePurchaseGooglePlay::IsAllowedToPurchase(const FUniqueNetId& UserId)
 {
-	UE_LOG_ONLINE_PURCHASE(Verbose, TEXT("FOnlinePurchaseGooglePlay::IsAllowedToPurchase"));
-	extern bool AndroidThunkCpp_Iap_IsAllowedToMakePurchases();
-	return AndroidThunkCpp_Iap_IsAllowedToMakePurchases();
+	return true;
 }
 
 bool FOnlinePurchaseGooglePlay::IsAllowedToPurchase(const FUniqueNetId& UserId, FOnlineError& Error)
 {
-	bool bResult = IsAllowedToPurchase(UserId);
-	Error =  FOnlineError(bResult ? EOnlineErrorResult::Success : EOnlineErrorResult::Unknown) ;
-	return bResult;
+	Error =  FOnlineError(EOnlineErrorResult::Success);
+	return true;
+}
+
+bool FOnlinePurchaseGooglePlay::IsBillingClientConnected()
+{
+	extern bool AndroidThunkCpp_Iap_IsAllowedToMakePurchases();
+	bool bIsBillingClientConnected = AndroidThunkCpp_Iap_IsAllowedToMakePurchases();
+	UE_LOG_ONLINE_PURCHASE(Verbose, TEXT("FOnlinePurchaseGooglePlay::IsBillingClientConnected: %d"), bIsBillingClientConnected);
+	return bIsBillingClientConnected;
 }
 
 void FOnlinePurchaseGooglePlay::Checkout(const FUniqueNetId& UserId, const FPurchaseCheckoutRequest& CheckoutRequest, const FOnPurchaseCheckoutComplete& Delegate)
@@ -240,13 +245,13 @@ void FOnlinePurchaseGooglePlay::Checkout(const FUniqueNetId& UserId, const FPurc
 
 	if (InProgressTransaction != nullptr)
 	{
-		UE_LOG_ONLINE_PURCHASE(Warning, TEXT("Previous purchase finished without notification"));			
+		UE_LOG_ONLINE_PURCHASE(Warning, TEXT("Previous purchase finished without notification"));
 		InProgressTransaction.Reset();
 	}
 
 	TSharedPtr<FOnlinePurchaseInProgressTransaction> RequestedTransaction = MakeShared<FOnlinePurchaseInProgressTransaction>(CheckoutRequest, UserId, EPurchaseTransactionState::NotStarted, Delegate);
 
-	if (IsAllowedToPurchase(UserId))
+	if (IsBillingClientConnected())
 	{
 		int32 NumOffers = CheckoutRequest.PurchaseOffers.Num();
 		if (NumOffers > 0)
@@ -268,6 +273,7 @@ void FOnlinePurchaseGooglePlay::Checkout(const FUniqueNetId& UserId, const FPurc
 			else
 			{
 				UE_LOG_ONLINE_PURCHASE(Display, TEXT("Failed to create a transaction."));
+				ErrorMessage = NSLOCTEXT("GooglePlayPurchase", "ErrorFailedToBeginPurchase", "Failed to checkout, failed to begin the purchase.");
 				RequestedTransaction->PurchaseInfo.TransactionState = EPurchaseTransactionState::Failed;
 			}
 		}
@@ -279,7 +285,7 @@ void FOnlinePurchaseGooglePlay::Checkout(const FUniqueNetId& UserId, const FPurc
 	}
 	else
 	{
-		ErrorMessage = NSLOCTEXT("GooglePlayPurchase", "ErrorPurchaseNotAllowed", "Failed to checkout, user not allowed to purchase.");
+		ErrorMessage = NSLOCTEXT("GooglePlayPurchase", "ErrorBillingClientNotConnected", "Failed to checkout, the billing client is not connected.");
 		RequestedTransaction->PurchaseInfo.TransactionState = EPurchaseTransactionState::Failed;
 	}
 
diff --git a/Engine/Plugins/Online/Android/OnlineSubsystemGooglePlay/Source/Public/OnlinePurchaseGooglePlay.h b/Engine/Plugins/Online/Android/OnlineSubsystemGooglePlay/Source/Public/OnlinePurchaseGooglePlay.h
--- a/Engine/Plugins/Online/Android/OnlineSubsystemGooglePlay/Source/Public/OnlinePurchaseGooglePlay.h
+++ b/Engine/Plugins/Online/Android/OnlineSubsystemGooglePlay/Source/Public/OnlinePurchaseGooglePlay.h
@@ -125,6 +125,9 @@ public:
 
 	static bool IsSubscriptionProductId(const FString& ProductId);
 private:
+
+	/** Returns true if the Billing Client is connected and ready */
+	bool IsBillingClientConnected();
 	
 	/**
 	 * Info used to cache and track orders in progress.
diff --git a/Engine/Plugins/Online/OnlineSubsystem/Source/Public/Interfaces/OnlinePurchaseInterface.h b/Engine/Plugins/Online/OnlineSubsystem/Source/Public/Interfaces/OnlinePurchaseInterface.h
--- a/Engine/Plugins/Online/OnlineSubsystem/Source/Public/Interfaces/OnlinePurchaseInterface.h
+++ b/Engine/Plugins/Online/OnlineSubsystem/Source/Public/Interfaces/OnlinePurchaseInterface.h
@@ -255,22 +255,22 @@ public:
 	virtual ~IOnlinePurchase() = default;
 
 	/**
-	 * Determine if user is allowed to purchase from store 
+	 * Determine if user is allowed to purchase on this platform from a policy or permission perspective
 	 *
 	 * @param UserId user initiating the request
 	 *
-	 * @return true if user can make a purchase
+	 * @return true if user is allowed to make a purchase
 	 */
 	virtual bool IsAllowedToPurchase(const FUniqueNetId& UserId) = 0;
 
 	/**
-	 * Determine if user is allowed to purchase from store
+	 * Determine if user is allowed to purchase from store on this platform from a policy or permission perspective
 	 *
 	 * @param UserId user initiating the request
 	 * 
-	 * @param Error an error result indicating why the the user isn't allowed when relevant and avaialable.
+	 * @param Error an error result indicating why the the user isn't allowed when relevant and available.
 	 *
-	 * @return true if user can make a purchase
+	 * @return true if user is allowed to make a purchase
 	 */
 	virtual bool IsAllowedToPurchase(const FUniqueNetId& UserId, FOnlineError& Error) = 0;
 
