[Android] Add a method to get the GooglePlay Billing CountryCode (FAndroidMisc::GetBillingConfigCountryCode)
#rnx

diff --git a/Engine/Build/Android/Java/src/com/epicgames/unreal/GameActivity.java.template b/Engine/Build/Android/Java/src/com/epicgames/unreal/GameActivity.java.template
--- a/Engine/Build/Android/Java/src/com/epicgames/unreal/GameActivity.java.template
+++ b/Engine/Build/Android/Java/src/com/epicgames/unreal/GameActivity.java.template
@@ -6425,6 +6425,19 @@ Trace.endSection();
 		}
 		return bTriggeredQuery;
 	}
+	
+	public String AndroidThunkJava_IapGetBillingCountryCode()
+	{
+		Log.debug("[JAVA] - AndroidThunkJava_IapGetBillingCountryCode");
+		if (IapStoreHelper != null)
+		{
+			final String countryCode = IapStoreHelper.GetBillingCountryCode();
+			Log.debug("countryCode=" + countryCode);
+			return countryCode;
+		}
+		Log.debug("Invalid IapStoreHelper.");
+		return "";
+	}
 
 	@Override
 	protected void onActivityResult(int requestCode, int resultCode, Intent data)
diff --git a/Engine/Build/Android/Java/src/com/epicgames/unreal/StoreHelper.java b/Engine/Build/Android/Java/src/com/epicgames/unreal/StoreHelper.java
--- a/Engine/Build/Android/Java/src/com/epicgames/unreal/StoreHelper.java
+++ b/Engine/Build/Android/Java/src/com/epicgames/unreal/StoreHelper.java
@@ -13,5 +13,6 @@ public interface StoreHelper
 	public void AcknowledgePurchase(String purchaseToken);
 	public void ConsumePurchase(String purchaseToken);
 	public boolean QueryExistingPurchases();
+	public String GetBillingCountryCode();
 	public void	onDestroy();
 }
\ No newline at end of file
diff --git a/Engine/Plugins/Online/Android/OnlineSubsystemGooglePlay/Source/Java/com/epicgames/unreal/GooglePlayStoreHelper.java b/Engine/Plugins/Online/Android/OnlineSubsystemGooglePlay/Source/Java/com/epicgames/unreal/GooglePlayStoreHelper.java
--- a/Engine/Plugins/Online/Android/OnlineSubsystemGooglePlay/Source/Java/com/epicgames/unreal/GooglePlayStoreHelper.java
+++ b/Engine/Plugins/Online/Android/OnlineSubsystemGooglePlay/Source/Java/com/epicgames/unreal/GooglePlayStoreHelper.java
@@ -9,12 +9,15 @@ import androidx.annotation.Nullable;
 
 import com.android.billingclient.api.BillingClient;
 import com.android.billingclient.api.BillingClientStateListener;
+import com.android.billingclient.api.BillingConfig;
+import com.android.billingclient.api.BillingConfigResponseListener;
 import com.android.billingclient.api.BillingFlowParams;
 import com.android.billingclient.api.BillingResult;
 import com.android.billingclient.api.AcknowledgePurchaseParams;
 import com.android.billingclient.api.AcknowledgePurchaseResponseListener;
 import com.android.billingclient.api.ConsumeParams;
 import com.android.billingclient.api.ConsumeResponseListener;
+import com.android.billingclient.api.GetBillingConfigParams;
 import com.android.billingclient.api.PendingPurchasesParams;
 import com.android.billingclient.api.ProductDetails;
 import com.android.billingclient.api.ProductDetailsResponseListener;
@@ -38,6 +41,9 @@ public class GooglePlayStoreHelper implements StoreHelper, PurchasesUpdatedListe
 {
 	// Our IAB helper interface provided by google.
 	private final BillingClient Client;
+	
+	// The IAP Billing config.
+	private BillingConfig ClientBillingConfig;
 
 	// Flag that determines whether the store is ready to use.
 	private boolean bIsIapSetup = false;
@@ -73,6 +79,19 @@ public class GooglePlayStoreHelper implements StoreHelper, PurchasesUpdatedListe
 				{
 					bIsIapSetup = true;
 					Log.debug("[GooglePlayStoreHelper] - In-app billing supported for " + MainActivity.getPackageName());
+
+					Client.getBillingConfigAsync(GetBillingConfigParams.newBuilder().build(), new BillingConfigResponseListener()
+					{
+						@Override
+						public void onBillingConfigResponse(BillingResult billingResult, BillingConfig billingConfig)
+						{
+							if (billingResult.getResponseCode() == BillingClient.BillingResponseCode.OK && billingConfig != null)
+							{
+								Log.debug("[GooglePlayStoreHelper] - Billing Config is now valid");
+								ClientBillingConfig = billingConfig;
+							}
+						}
+					});
 				}
 				else
 				{
@@ -88,6 +107,18 @@ public class GooglePlayStoreHelper implements StoreHelper, PurchasesUpdatedListe
 			}
 		});
 	}
+	
+	@Override
+	public String GetBillingCountryCode()
+	{
+		Log.debug("[GooglePlayStoreHelper] - GetBillingCountryCode");
+		if (ClientBillingConfig != null)
+		{
+			return ClientBillingConfig.getCountryCode();
+		}
+		Log.debug("[GooglePlayStoreHelper] - ClientBillingConfig is invalid. Can't get the country code.");
+		return "";
+	}
 
 	///////////////////////////////////////////////////////
 	// The StoreHelper interfaces implementation for Google Play Store.
diff --git a/Engine/Source/Runtime/Core/Private/Android/AndroidPlatformMisc.cpp b/Engine/Source/Runtime/Core/Private/Android/AndroidPlatformMisc.cpp
--- a/Engine/Source/Runtime/Core/Private/Android/AndroidPlatformMisc.cpp
+++ b/Engine/Source/Runtime/Core/Private/Android/AndroidPlatformMisc.cpp
@@ -907,6 +907,15 @@ bool FAndroidMisc::DeleteStoredSection(const FString& InStoreId, const FString&
 	return FGenericPlatformMisc::DeleteStoredSection(InStoreId, InSectionName);
 }
 
+FString FAndroidMisc::GetBillingConfigCountryCode()
+{
+#if USE_ANDROID_JNI
+	extern FString AndroidThunkJava_IapGetBillingCountryCode();
+	return AndroidThunkJava_IapGetBillingCountryCode(); 
+#endif
+	return FString();
+}
+
 bool FAndroidMisc::UseRenderThread()
 {
 	// if we in general don't want to use the render thread due to commandline, etc, then don't
@@ -2516,7 +2525,6 @@ void JNICALL UE::Jni::FGameActivity::nativeSetAndroidStartupState(JNIEnv* jenv,
 		bDetectedDebugger = true;
 	}
 }
-
 #endif
 bool FAndroidMisc::UseNewWindowBehavior()
 {
diff --git a/Engine/Source/Runtime/Core/Public/Android/AndroidPlatformMisc.h b/Engine/Source/Runtime/Core/Public/Android/AndroidPlatformMisc.h
--- a/Engine/Source/Runtime/Core/Public/Android/AndroidPlatformMisc.h
+++ b/Engine/Source/Runtime/Core/Public/Android/AndroidPlatformMisc.h
@@ -66,6 +66,8 @@ struct FAndroidMisc : public FGenericPlatformMisc
 	static CORE_API bool DeleteStoredValue(const FString& InStoreId, const FString& InSectionName, const FString& InKeyName);
 	static CORE_API bool DeleteStoredSection(const FString& InStoreId, const FString& InSectionName);
 
+	static CORE_API FString GetBillingConfigCountryCode();
+
 public:
 
 	static CORE_API bool AllowThreadHeartBeat();
diff --git a/Engine/Source/Runtime/Launch/Private/Android/AndroidJNI.cpp b/Engine/Source/Runtime/Launch/Private/Android/AndroidJNI.cpp
--- a/Engine/Source/Runtime/Launch/Private/Android/AndroidJNI.cpp
+++ b/Engine/Source/Runtime/Launch/Private/Android/AndroidJNI.cpp
@@ -242,6 +242,7 @@ void FJavaWrapper::FindGooglePlayBillingMethods(JNIEnv* Env)
 	AndroidThunkJava_IapAcknowledgePurchase = FindMethod(Env, GoogleServicesClassID, "AndroidThunkJava_IapAcknowledgePurchase", "(Ljava/lang/String;)Z", bIsStoreOptional);
 	AndroidThunkJava_IapConsumePurchase = FindMethod(Env, GoogleServicesClassID, "AndroidThunkJava_IapConsumePurchase", "(Ljava/lang/String;)Z", bIsStoreOptional);
 	AndroidThunkJava_IapQueryExistingPurchases = FindMethod(Env, GoogleServicesClassID, "AndroidThunkJava_IapQueryExistingPurchases", "()Z", bIsStoreOptional);
+	AndroidThunkJava_IapGetBillingCountryCode = FindMethod(Env, GoogleServicesClassID, "AndroidThunkJava_IapGetBillingCountryCode", "()Ljava/lang/String;", bIsStoreOptional);
 }
 
 jclass FJavaWrapper::FindClass(JNIEnv* Env, const ANSICHAR* ClassName, bool bIsOptional)
@@ -588,6 +589,7 @@ jmethodID FJavaWrapper::AndroidThunkJava_IapIsAllowedToMakePurchases;
 jmethodID FJavaWrapper::AndroidThunkJava_IapQueryExistingPurchases;
 jmethodID FJavaWrapper::AndroidThunkJava_IapAcknowledgePurchase;
 jmethodID FJavaWrapper::AndroidThunkJava_IapConsumePurchase;
+jmethodID FJavaWrapper::AndroidThunkJava_IapGetBillingCountryCode;
 
 jmethodID FJavaWrapper::AndroidThunkJava_UseSurfaceViewWorkaround;
 jmethodID FJavaWrapper::AndroidThunkJava_SetDesiredViewSize;
@@ -1816,6 +1818,17 @@ bool AndroidThunkCpp_Iap_IsAllowedToMakePurchases()
 	return bResult;
 }
 
+FString AndroidThunkJava_IapGetBillingCountryCode()
+{
+	FString Result = FString();
+	if (JNIEnv* Env = FAndroidApplication::GetJavaEnv())
+	{
+		Result = FJavaHelper::FStringFromLocalRef(Env, (jstring)FJavaWrapper::CallObjectMethod(Env, FJavaWrapper::GameActivityThis, FJavaWrapper::AndroidThunkJava_IapGetBillingCountryCode));
+	}
+
+	return Result;
+}
+
 void AndroidThunkCpp_UseSurfaceViewWorkaround()
 {
 	if (JNIEnv* Env = FAndroidApplication::GetJavaEnv())
diff --git a/Engine/Source/Runtime/Launch/Public/Android/AndroidJNI.h b/Engine/Source/Runtime/Launch/Public/Android/AndroidJNI.h
--- a/Engine/Source/Runtime/Launch/Public/Android/AndroidJNI.h
+++ b/Engine/Source/Runtime/Launch/Public/Android/AndroidJNI.h
@@ -129,6 +129,7 @@ public:
 	static jmethodID AndroidThunkJava_IapQueryExistingPurchases;
 	static jmethodID AndroidThunkJava_IapAcknowledgePurchase;
 	static jmethodID AndroidThunkJava_IapConsumePurchase;
+	static jmethodID AndroidThunkJava_IapGetBillingCountryCode;
 
 	// SurfaceView functionality for view scaling on some devices
 	static jmethodID AndroidThunkJava_UseSurfaceViewWorkaround;
