Change Google Play Billing errors to display additional debug info, get rid of delay for BILLING_UNAVAILABLE, don't display an error dialog for BILLING_UNAVAILABLE since google play most likely displayed one already

#rnx
#tests local, QA

diff --git a/Engine/Plugins/Online/Android/OnlineSubsystemGooglePlay/Source/Private/OnlinePurchaseGooglePlay.cpp b/Engine/Plugins/Online/Android/OnlineSubsystemGooglePlay/Source/Private/OnlinePurchaseGooglePlay.cpp
--- a/Engine/Plugins/Online/Android/OnlineSubsystemGooglePlay/Source/Private/OnlinePurchaseGooglePlay.cpp
+++ b/Engine/Plugins/Online/Android/OnlineSubsystemGooglePlay/Source/Private/OnlinePurchaseGooglePlay.cpp
@@ -432,8 +432,12 @@ void FOnlinePurchaseGooglePlay::OnTransactionCompleteResponse(EGooglePlayBilling
 		switch (Result)
 		{
 			case EPurchaseTransactionState::Failed:
-				FinalResult.SetFromErrorCode(TEXT("com.epicgames.purchase.failure"));
+				FinalResult.SetFromErrorCode(FString::Printf(TEXT("com.epicgames.purchase.googleplay.failure[%d:%s]"), InResponseCode, LexToString(InResponseCode)));
 				FinalResult.ErrorMessage = !ErrorStr.IsEmpty() ? FText::FromString(ErrorStr) : LOCTEXT("GooglePlayTransactionFailed", "Transaction Failed");
+				if (InResponseCode == EGooglePlayBillingResponseCode::BillingUnavailable)
+				{
+					FinalResult.SetErrorRaw(FString::Printf(TEXT("%s.noDelayedRedeemAndError"), *FinalResult.GetErrorCode()));
+				}
 				break;
 			case EPurchaseTransactionState::Canceled:
 				FinalResult.SetFromErrorCode(TEXT("com.epicgames.catalog_helper.user_cancelled"));
@@ -447,7 +451,7 @@ void FOnlinePurchaseGooglePlay::OnTransactionCompleteResponse(EGooglePlayBilling
 				FinalResult.ErrorMessage = !ErrorStr.IsEmpty() ? FText::FromString(ErrorStr) : LOCTEXT("GooglePlayTransactionDeferred", "Transaction Deferred");
 				break;
 			case EPurchaseTransactionState::Invalid:
-				FinalResult.SetFromErrorCode(TEXT("com.epicgames.purchase.invalid"));
+				FinalResult.SetFromErrorCode(FString::Printf(TEXT("com.epicgames.purchase.googleplay.invalid[%d:%s]"), InResponseCode, LexToString(InResponseCode)));
 				FinalResult.ErrorMessage = !ErrorStr.IsEmpty() ? FText::FromString(ErrorStr) : LOCTEXT("GooglePlayInvalidState", "Invalid purchase result");
 				break;
 			default:
