Update Google Play Store Online Subsystem to support current requirements for alternative billing - add ability to show the informational dialog and check if only alternative billing is available

#rnx

diff --git a/Engine/Build/Android/Java/src/com/epicgames/unreal/GameActivity.java.template b/Engine/Build/Android/Java/src/com/epicgames/unreal/GameActivity.java.template
--- a/Engine/Build/Android/Java/src/com/epicgames/unreal/GameActivity.java.template
+++ b/Engine/Build/Android/Java/src/com/epicgames/unreal/GameActivity.java.template
@@ -6393,7 +6393,7 @@ Trace.endSection();
 	}
 	
 	// In app purchase functionality
-	public void AndroidThunkJava_IapSetupService(String InProductKey)
+	public void AndroidThunkJava_IapSetupService(String InProductKey, boolean InEnableAlternativeBillingOnly)
 	{
 //$${gameActivityIapSetupServiceAdditions}$$
 	}
@@ -6439,6 +6439,30 @@ Trace.endSection();
 		return "";
 	}
 
+	public boolean AndroidThunkJava_IapGetIsAlternativeBillingOnlyAvailable()
+	{
+		Log.debug("[JAVA] - AndroidThunkJava_IapGetIsAlternativeBillingOnlyAvailable");
+		if (IapStoreHelper != null)
+		{
+			return IapStoreHelper.GetIsAlternativeBillingOnlyAvailable();
+		}
+
+		Log.debug("[JAVA] - Store Helper is invalid");
+		return true;
+	}
+
+	public boolean AndroidThunkJava_IapShowAlternativeBillingOnlyInformationDialog()
+	{
+		Log.debug("[JAVA] - AndroidThunkJava_IapShowAlternativeBillingOnlyInformationDialog");
+		if (IapStoreHelper != null)
+		{
+			return IapStoreHelper.ShowAlternativeBillingOnlyInformationDialog();
+		}
+
+		Log.debug("[JAVA] - Store Helper is invalid");
+		return false;
+	}
+
 	@Override
 	protected void onActivityResult(int requestCode, int resultCode, Intent data)
 	{
diff --git a/Engine/Build/Android/Java/src/com/epicgames/unreal/StoreHelper.java b/Engine/Build/Android/Java/src/com/epicgames/unreal/StoreHelper.java
--- a/Engine/Build/Android/Java/src/com/epicgames/unreal/StoreHelper.java
+++ b/Engine/Build/Android/Java/src/com/epicgames/unreal/StoreHelper.java
@@ -14,5 +14,7 @@ public interface StoreHelper
 	public void ConsumePurchase(String purchaseToken);
 	public boolean QueryExistingPurchases();
 	public String GetBillingCountryCode();
+	public boolean GetIsAlternativeBillingOnlyAvailable();
+	public boolean ShowAlternativeBillingOnlyInformationDialog();
 	public void	onDestroy();
 }
\ No newline at end of file
diff --git a/Engine/Plugins/Online/Android/OnlineSubsystemGooglePlay/Source/Java/com/epicgames/unreal/GooglePlayStoreHelper.java b/Engine/Plugins/Online/Android/OnlineSubsystemGooglePlay/Source/Java/com/epicgames/unreal/GooglePlayStoreHelper.java
--- a/Engine/Plugins/Online/Android/OnlineSubsystemGooglePlay/Source/Java/com/epicgames/unreal/GooglePlayStoreHelper.java
+++ b/Engine/Plugins/Online/Android/OnlineSubsystemGooglePlay/Source/Java/com/epicgames/unreal/GooglePlayStoreHelper.java
@@ -15,6 +15,8 @@ import com.android.billingclient.api.BillingFlowParams;
 import com.android.billingclient.api.BillingResult;
 import com.android.billingclient.api.AcknowledgePurchaseParams;
 import com.android.billingclient.api.AcknowledgePurchaseResponseListener;
+import com.android.billingclient.api.AlternativeBillingOnlyAvailabilityListener;
+import com.android.billingclient.api.AlternativeBillingOnlyInformationDialogListener;
 import com.android.billingclient.api.ConsumeParams;
 import com.android.billingclient.api.ConsumeResponseListener;
 import com.android.billingclient.api.GetBillingConfigParams;
@@ -48,13 +50,19 @@ public class GooglePlayStoreHelper implements StoreHelper, PurchasesUpdatedListe
 	// Flag that determines whether the store is ready to use.
 	private boolean bIsIapSetup = false;
 
+	// Flag of if alternative billing only was set up
+	private boolean bIsAlternativeBillingOnlyActive = false;
+
+	// Cached result of isAlternativeBillingOnlyAvailableAsync
+	private boolean bIsAlternativeBillingOnlyAvailable = false;
+
 	// Output device for log messages.
 	private final Logger Log;
 
 	// Cache access to the games activity.
 	private final GameActivity MainActivity;
 
-	public GooglePlayStoreHelper(GameActivity InGameActivity, Logger InLog)
+	public GooglePlayStoreHelper(GameActivity InGameActivity, Logger InLog, boolean InEnableAlternativeBillingOnly)
 	{
 		Log = InLog;
 		Log.debug("[GooglePlayStoreHelper] - GooglePlayStoreHelper");
@@ -65,10 +73,18 @@ public class GooglePlayStoreHelper implements StoreHelper, PurchasesUpdatedListe
 			.enableOneTimeProducts()
 			.build();
 
-		Client = BillingClient.newBuilder(MainActivity)
+		BillingClient.Builder BillingClientBuilder = BillingClient.newBuilder(MainActivity)
 			.setListener(this)
 			.enablePendingPurchases(PendingPurchases)
-			.enableAutoServiceReconnection()
+			.enableAutoServiceReconnection();
+		
+		if (InEnableAlternativeBillingOnly)
+		{
+			bIsAlternativeBillingOnlyActive = true;
+			BillingClientBuilder.enableAlternativeBillingOnly();
+		}
+		
+		Client = BillingClientBuilder
 			.build();
 
 		Client.startConnection(new BillingClientStateListener() {
@@ -92,6 +108,19 @@ public class GooglePlayStoreHelper implements StoreHelper, PurchasesUpdatedListe
 							}
 						}
 					});
+
+					if (bIsAlternativeBillingOnlyActive)
+					{
+						Client.isAlternativeBillingOnlyAvailableAsync(new AlternativeBillingOnlyAvailabilityListener() {
+							@Override
+							public void onAlternativeBillingOnlyAvailabilityResponse(
+								BillingResult billingResult) 
+								{
+									bIsAlternativeBillingOnlyAvailable = billingResult.getResponseCode() == BillingClient.BillingResponseCode.OK;
+									Log.debug("[GooglePlayStoreHelper] - bIsAlternativeBillingOnlyAvailable: " + bIsAlternativeBillingOnlyAvailable + " error " + TranslateBillingResult(billingResult));
+								}
+						});
+					}
 				}
 				else
 				{
@@ -120,6 +149,41 @@ public class GooglePlayStoreHelper implements StoreHelper, PurchasesUpdatedListe
 		return "";
 	}
 
+	@Override
+	public boolean GetIsAlternativeBillingOnlyAvailable()
+	{
+		Log.debug("[GooglePlayStoreHelper] - GetIsAlternativeBillingOnlyAvailable");
+		
+		return bIsAlternativeBillingOnlyAvailable;
+	}
+
+	@Override
+	public boolean ShowAlternativeBillingOnlyInformationDialog()
+	{
+		if (!bIsAlternativeBillingOnlyActive)
+		{
+			return true;
+		}
+
+		AlternativeBillingOnlyInformationDialogListener listener =
+			new AlternativeBillingOnlyInformationDialogListener() {
+				@Override
+				public void onAlternativeBillingOnlyInformationDialogResponse(
+					BillingResult billingResult) {
+						Log.debug("[GooglePlayStoreHelper] - NativeAlternativeBillingOnlyInformationDialogComplete");
+						NativeAlternativeBillingOnlyInformationDialogComplete(billingResult.getResponseCode());
+					}
+			};
+
+		Log.debug("[GooglePlayStoreHelper] - showAlternativeBillingOnlyInformationDialog");
+		BillingResult billingResult =
+			Client.showAlternativeBillingOnlyInformationDialog(MainActivity,
+				listener);
+
+		Log.debug("[GooglePlayStoreHelper] - showAlternativeBillingOnlyInformationDialog return " + TranslateBillingResult(billingResult));
+		return billingResult.getResponseCode() == BillingClient.BillingResponseCode.OK;
+	}
+
 	///////////////////////////////////////////////////////
 	// The StoreHelper interfaces implementation for Google Play Store.
 
@@ -903,4 +967,5 @@ public class GooglePlayStoreHelper implements StoreHelper, PurchasesUpdatedListe
 	public native void NativeQueryExistingPurchasesComplete(int ResponseCode, String[][] ProductIds, int[] PurchaseState, String[] ProductTokens, String[] ReceiptsData, String[] Signatures);
 	public native void NativeConsumeComplete(int ResponseCode, String PurchaseToken);
 	public native void NativeAcknowledgeComplete(int ResponseCode, String PurchaseToken);
+	public native void NativeAlternativeBillingOnlyInformationDialogComplete(int ResponseCode);
 }
diff --git a/Engine/Plugins/Online/Android/OnlineSubsystemGooglePlay/Source/OnlineSubsystemGooglePlay_UPL.xml b/Engine/Plugins/Online/Android/OnlineSubsystemGooglePlay/Source/OnlineSubsystemGooglePlay_UPL.xml
--- a/Engine/Plugins/Online/Android/OnlineSubsystemGooglePlay/Source/OnlineSubsystemGooglePlay_UPL.xml
+++ b/Engine/Plugins/Online/Android/OnlineSubsystemGooglePlay/Source/OnlineSubsystemGooglePlay_UPL.xml
@@ -158,7 +158,7 @@ import java.security.NoSuchAlgorithmException;
 						<insert>
 		if (getPackageManager().checkPermission("com.android.vending.BILLING", getPackageName()) == getPackageManager().PERMISSION_GRANTED)
 		{
-			IapStoreHelper = new GooglePlayStoreHelper(this, Log);
+			IapStoreHelper = new GooglePlayStoreHelper(this, Log, InEnableAlternativeBillingOnly);
 			if (IapStoreHelper != null)
 			{
 				Log.debug("[JAVA] - AndroidThunkJava_IapSetupService - Setup started");
diff --git a/Engine/Plugins/Online/Android/OnlineSubsystemGooglePlay/Source/Private/OnlineJniGooglePlayStoreHelper.h b/Engine/Plugins/Online/Android/OnlineSubsystemGooglePlay/Source/Private/OnlineJniGooglePlayStoreHelper.h
--- a/Engine/Plugins/Online/Android/OnlineSubsystemGooglePlay/Source/Private/OnlineJniGooglePlayStoreHelper.h
+++ b/Engine/Plugins/Online/Android/OnlineSubsystemGooglePlay/Source/Private/OnlineJniGooglePlayStoreHelper.h
@@ -19,6 +19,7 @@ namespace UE::Jni
 		static void JNICALL NativePurchaseComplete(JNIEnv* env, jobject thiz, jint JavaResponseCode, Java::Lang::TArray<jstring>* JavaProductIds, jint JavaPurchaseState, jstring JavaPurchaseToken, jstring JavaReceiptData, jstring JavaSignature);
 		static void JNICALL NativeConsumeComplete(JNIEnv* env, jobject thiz, jint JavaResponseCode, jstring JavaPurchaseToken);
 		static void JNICALL NativeAcknowledgeComplete(JNIEnv* env, jobject thiz, jint JavaResponseCode, jstring JavaPurchaseToken);
+		static void JNICALL NativeAlternativeBillingOnlyInformationDialogComplete(JNIEnv* env, jobject thiz, jint JavaResponseCode);
 
 		static constexpr FNativeMethod NativeMethods[]
 		{
@@ -26,7 +27,8 @@ namespace UE::Jni
 			UE_JNI_NATIVE_METHOD(NativeQueryExistingPurchasesComplete),
 			UE_JNI_NATIVE_METHOD(NativePurchaseComplete),
 			UE_JNI_NATIVE_METHOD(NativeConsumeComplete),
-			UE_JNI_NATIVE_METHOD(NativeAcknowledgeComplete)
+			UE_JNI_NATIVE_METHOD(NativeAcknowledgeComplete),
+			UE_JNI_NATIVE_METHOD(NativeAlternativeBillingOnlyInformationDialogComplete)
 		};
 	};
 	
diff --git a/Engine/Plugins/Online/Android/OnlineSubsystemGooglePlay/Source/Private/OnlinePurchaseGooglePlay.cpp b/Engine/Plugins/Online/Android/OnlineSubsystemGooglePlay/Source/Private/OnlinePurchaseGooglePlay.cpp
--- a/Engine/Plugins/Online/Android/OnlineSubsystemGooglePlay/Source/Private/OnlinePurchaseGooglePlay.cpp
+++ b/Engine/Plugins/Online/Android/OnlineSubsystemGooglePlay/Source/Private/OnlinePurchaseGooglePlay.cpp
@@ -210,8 +210,11 @@ void FOnlinePurchaseGooglePlay::Init()
 
 	GConfig->GetBool(TEXT("OnlineSubsystemGooglePlay.Store"), TEXT("bDisableLocalAcknowledgeAndConsume"), bDisableLocalAcknowledgeAndConsume, GEngineIni);
 
-	extern void AndroidThunkCpp_Iap_SetupIapService(const FString&);
-	AndroidThunkCpp_Iap_SetupIapService(GooglePlayLicenseKey);
+	bool EnableAlternativeBillingOnly = false;
+	GConfig->GetBool(TEXT("OnlineSubsystemGooglePlay.Store"), TEXT("EnableAlternativeBillingOnly"), EnableAlternativeBillingOnly, GEngineIni);
+
+	extern void AndroidThunkCpp_Iap_SetupIapService(const FString&, bool);
+	AndroidThunkCpp_Iap_SetupIapService(GooglePlayLicenseKey, EnableAlternativeBillingOnly);
 }
 
 bool FOnlinePurchaseGooglePlay::IsSubscriptionProductId(const FString& ProductId)
@@ -409,6 +412,65 @@ void FOnlinePurchaseGooglePlay::GetReceipts(const FUniqueNetId& UserId, TArray<F
 	UE_LOG_ONLINE_PURCHASE(Log, TEXT("FOnlinePurchaseGooglePlay::GetReceipts - Final Number of OutReceipts - %d"), OutReceipts.Num());
 }
 
+void FOnlinePurchaseGooglePlay::ShowAlternativeBillingInformationDialog(const FOnShowAlternativeBillingInformationDialogComplete& Delegate)
+{
+	UE_LOG_ONLINE_PURCHASE(Verbose, TEXT("FOnlinePurchaseGooglePlay::ShowAlternativeBillingOnlyInformationDialog"));
+
+	if (!bRequestedAlternativeBillingOnlyInformationDialog)
+	{
+		bRequestedAlternativeBillingOnlyInformationDialog = true;
+		AlternativeBillingOnlyInformationDialogComplete = Delegate;
+
+		extern bool AndroidThunkJava_IapShowAlternativeBillingOnlyInformationDialog();
+		if (!AndroidThunkJava_IapShowAlternativeBillingOnlyInformationDialog())
+		{
+			UE_LOG_ONLINE_PURCHASE(Verbose, TEXT("FOnlinePurchaseGooglePlay::ShowAlternativeBillingInformationDialog failed to start dialog"));
+
+			Subsystem->ExecuteNextTick([this]() {
+				FOnlineError Result(EOnlineErrorResult::RequestFailure);
+				AlternativeBillingOnlyInformationDialogComplete.ExecuteIfBound(Result);
+				AlternativeBillingOnlyInformationDialogComplete.Unbind();
+				bRequestedAlternativeBillingOnlyInformationDialog = false;
+			});
+		}
+	}
+	else
+	{
+
+		UE_LOG_ONLINE_PURCHASE(Verbose, TEXT("FOnlinePurchaseGooglePlay::ShowAlternativeBillingInformationDialog already in progress."));
+
+		Subsystem->ExecuteNextTick([Delegate = Delegate]() mutable {
+			FOnlineError Result(false);
+			Delegate.ExecuteIfBound(Result);
+		});
+	}
+}
+
+void FOnlinePurchaseGooglePlay::OnAlternativeBillingOnlyInformationDialogComplete(EGooglePlayBillingResponseCode InResponseCode)
+{
+	UE_LOG_ONLINE_PURCHASE(Log, TEXT("FOnlinePurchaseGooglePlay::OnAlternativeBillingOnlyInformationDialogComplete %s"), LexToString(InResponseCode));
+
+	if (bRequestedAlternativeBillingOnlyInformationDialog)
+	{
+		bool bSuccess = (InResponseCode == EGooglePlayBillingResponseCode::Ok);
+		FOnlineError Result(bSuccess);
+
+		if (!bSuccess)
+		{
+			Result.SetFromErrorCode(LexToString(InResponseCode));
+			UE_LOG_ONLINE_PURCHASE(Log, TEXT("OnAlternativeBillingOnlyInformationDialogComplete failed %s"), LexToString(InResponseCode));
+		}
+
+		AlternativeBillingOnlyInformationDialogComplete.ExecuteIfBound(Result);
+		AlternativeBillingOnlyInformationDialogComplete.Unbind();
+		bRequestedAlternativeBillingOnlyInformationDialog = false;
+	}
+	else
+	{
+		UE_LOG_ONLINE_PURCHASE(Warning, TEXT("FOnlinePurchaseGooglePlay::OnAlternativeBillingOnlyInformationDialogComplete unexpected call %s"), LexToString(InResponseCode));
+	}
+}
+
 void FOnlinePurchaseGooglePlay::OnTransactionCompleteResponse(EGooglePlayBillingResponseCode InResponseCode, const FGoogleTransactionData& InTransactionData)
 {
 	UE_LOG_ONLINE_PURCHASE(Log, TEXT("FOnlinePurchaseGooglePlay::OnTransactionCompleteResponse %s. Transaction: %s"), LexToString(InResponseCode), *InTransactionData.ToDebugString());
@@ -616,5 +678,21 @@ void JNICALL UE::Jni::FGooglePlayStoreHelper::NativeAcknowledgeComplete(JNIEnv*
 	UE_LOG_ONLINE_PURCHASE(Verbose, TEXT("NativeAcknowledgeComplete received response code: %s. PurchaseToken: %s\n"), LexToString(EGPResponse), *PurchaseToken);
 }
 
+void JNICALL UE::Jni::FGooglePlayStoreHelper::NativeAlternativeBillingOnlyInformationDialogComplete(JNIEnv* jenv, jobject /*Thiz*/, jint JavaResponseCode)
+{
+	EGooglePlayBillingResponseCode EGPResponse = (EGooglePlayBillingResponseCode)JavaResponseCode;
+
+	UE_LOG_ONLINE_PURCHASE(Verbose, TEXT("NativeAlternativeBillingOnlyInformationDialogComplete received response code: %s\n"), LexToString(EGPResponse));
+
+	if (auto OnlineSubGP = static_cast<FOnlineSubsystemGooglePlay* const>(IOnlineSubsystem::Get(GOOGLEPLAY_SUBSYSTEM)))
+	{
+		OnlineSubGP->ExecuteNextTick([OnlineSubGP, EGPResponse]()
+		{
+			TSharedPtr<FOnlinePurchaseGooglePlay> PurchaseInt = StaticCastSharedPtr<FOnlinePurchaseGooglePlay>(OnlineSubGP->GetPurchaseInterface());
+			PurchaseInt->OnAlternativeBillingOnlyInformationDialogComplete(EGPResponse);
+		});
+	}
+}
+
 #undef LOCTEXT_NAMESPACE
 
diff --git a/Engine/Plugins/Online/Android/OnlineSubsystemGooglePlay/Source/Private/OnlineStoreGooglePlayCommon.h b/Engine/Plugins/Online/Android/OnlineSubsystemGooglePlay/Source/Private/OnlineStoreGooglePlayCommon.h
--- a/Engine/Plugins/Online/Android/OnlineSubsystemGooglePlay/Source/Private/OnlineStoreGooglePlayCommon.h
+++ b/Engine/Plugins/Online/Android/OnlineSubsystemGooglePlay/Source/Private/OnlineStoreGooglePlayCommon.h
@@ -25,6 +25,7 @@ enum class EGooglePlayBillingResponseCode : int8
 	Error = 6,
 	ItemAlreadyOwned = 7,
 	ItemNotOwned = 8,
+	NetworkError = 12,
 };
 
 /* Possible transaction states returned from the Java GooglePlay billing interface.*/
@@ -63,6 +64,8 @@ inline const TCHAR* const LexToString(EGooglePlayBillingResponseCode InResponseC
 			return TEXT("DeveloperError");
 		case EGooglePlayBillingResponseCode::Error:
 			return TEXT("Error");
+		case EGooglePlayBillingResponseCode::NetworkError:
+			return TEXT("NetworkError");
 		default:
 			return TEXT("UnknownError");
 	}
diff --git a/Engine/Plugins/Online/Android/OnlineSubsystemGooglePlay/Source/Public/OnlinePurchaseGooglePlay.h b/Engine/Plugins/Online/Android/OnlineSubsystemGooglePlay/Source/Public/OnlinePurchaseGooglePlay.h
--- a/Engine/Plugins/Online/Android/OnlineSubsystemGooglePlay/Source/Public/OnlinePurchaseGooglePlay.h
+++ b/Engine/Plugins/Online/Android/OnlineSubsystemGooglePlay/Source/Public/OnlinePurchaseGooglePlay.h
@@ -117,6 +117,12 @@ public:
 	/** Initialize the interface */
 	void Init();
 
+	/** Show the required Google Play system alternative billing informational dialog */
+	virtual void ShowAlternativeBillingInformationDialog(const FOnShowAlternativeBillingInformationDialogComplete& Delegate) override;
+
+	/** Handle response from showing the Google Play alternative billing only information dialog **/
+	void OnAlternativeBillingOnlyInformationDialogComplete(EGooglePlayBillingResponseCode InResponseCode);
+
 	/** Handle Java side transaction completed notification */
 	void OnTransactionCompleteResponse(EGooglePlayBillingResponseCode InResponseCode, const FGoogleTransactionData& InTransactionData);
 
@@ -152,6 +158,12 @@ private:
 	/** Transient delegate to fire when query receipts has completed */
 	FOnQueryReceiptsComplete QueryReceiptsComplete;
 
+	/** Is the system dialog already being shown */
+	bool bRequestedAlternativeBillingOnlyInformationDialog = false;
+
+	/** Delegate to call when the system dialog is finished */
+	FOnShowAlternativeBillingInformationDialogComplete AlternativeBillingOnlyInformationDialogComplete;
+
 	/** Keeps track of pending user transactions */
 	TSharedPtr<FOnlinePurchaseInProgressTransaction> InProgressTransaction;
 	
diff --git a/Engine/Plugins/Online/OnlineSubsystem/Source/Public/Interfaces/OnlinePurchaseInterface.h b/Engine/Plugins/Online/OnlineSubsystem/Source/Public/Interfaces/OnlinePurchaseInterface.h
--- a/Engine/Plugins/Online/OnlineSubsystem/Source/Public/Interfaces/OnlinePurchaseInterface.h
+++ b/Engine/Plugins/Online/OnlineSubsystem/Source/Public/Interfaces/OnlinePurchaseInterface.h
@@ -235,6 +235,11 @@ DECLARE_DELEGATE_TwoParams(FOnPurchaseRedeemCodeComplete, const FOnlineError& /*
  */
 DECLARE_DELEGATE_OneParam(FOnQueryReceiptsComplete, const FOnlineError& /*Result*/);
 
+/**
+ * Delegate called when the information dialog shows and completes
+ */
+DECLARE_DELEGATE_OneParam(FOnShowAlternativeBillingInformationDialogComplete, const FOnlineError& /*Result*/);
+
 /**
 * Delegate called when receipt validation completes
 */
@@ -345,6 +350,13 @@ public:
 
 	virtual void FinalizeReceiptValidationInfo(const FUniqueNetId& UserId, FString& InReceiptValidationInfo, const FOnFinalizeReceiptValidationInfoComplete& Delegate) = 0;
 
+	/**
+	 * If the underlying IAP service requires an informational dialog to be shown, show it and trigger 
+	 * the delegate on completion with the status
+	 * @param Delegate completion callback (guaranteed to be called)
+	 */
+	virtual void ShowAlternativeBillingInformationDialog(const FOnShowAlternativeBillingInformationDialogComplete& Delegate) {}
+
 	/**
 	 * Delegate fired when the local system tells us of a new completed purchase we may not have initiated in-game.
 	 * Use this to know about new pending receipts in instances the local client did not start a purchase,
diff --git a/Engine/Source/Runtime/Core/Private/Android/AndroidPlatformMisc.cpp b/Engine/Source/Runtime/Core/Private/Android/AndroidPlatformMisc.cpp
--- a/Engine/Source/Runtime/Core/Private/Android/AndroidPlatformMisc.cpp
+++ b/Engine/Source/Runtime/Core/Private/Android/AndroidPlatformMisc.cpp
@@ -916,6 +916,15 @@ FString FAndroidMisc::GetBillingConfigCountryCode()
 	return FString();
 }
 
+bool FAndroidMisc::GetIsAlternativeBillingOnlyAvailable()
+{
+#if USE_ANDROID_JNI
+	extern bool AndroidThunkJava_IapGetIsAlternativeBillingOnlyAvailable();
+	return AndroidThunkJava_IapGetIsAlternativeBillingOnlyAvailable();
+#endif
+	return true;
+}
+
 bool FAndroidMisc::UseRenderThread()
 {
 	// if we in general don't want to use the render thread due to commandline, etc, then don't
diff --git a/Engine/Source/Runtime/Core/Public/Android/AndroidPlatformMisc.h b/Engine/Source/Runtime/Core/Public/Android/AndroidPlatformMisc.h
--- a/Engine/Source/Runtime/Core/Public/Android/AndroidPlatformMisc.h
+++ b/Engine/Source/Runtime/Core/Public/Android/AndroidPlatformMisc.h
@@ -67,6 +67,7 @@ struct FAndroidMisc : public FGenericPlatformMisc
 	static CORE_API bool DeleteStoredSection(const FString& InStoreId, const FString& InSectionName);
 
 	static CORE_API FString GetBillingConfigCountryCode();
+	static CORE_API bool GetIsAlternativeBillingOnlyAvailable();
 
 public:
 
diff --git a/Engine/Source/Runtime/Launch/Private/Android/AndroidJNI.cpp b/Engine/Source/Runtime/Launch/Private/Android/AndroidJNI.cpp
--- a/Engine/Source/Runtime/Launch/Private/Android/AndroidJNI.cpp
+++ b/Engine/Source/Runtime/Launch/Private/Android/AndroidJNI.cpp
@@ -235,7 +235,7 @@ void FJavaWrapper::FindGooglePlayBillingMethods(JNIEnv* Env)
 	bool bIsStoreOptional = !bSupportsInAppPurchasing;
 	
 	JavaStringClass = UE::Jni::Class<UE::Jni::Java::Lang::FString>;
-	AndroidThunkJava_IapSetupService = FindMethod(Env, GoogleServicesClassID, "AndroidThunkJava_IapSetupService", "(Ljava/lang/String;)V", bIsStoreOptional);
+	AndroidThunkJava_IapSetupService = FindMethod(Env, GoogleServicesClassID, "AndroidThunkJava_IapSetupService", "(Ljava/lang/String;Z)V", bIsStoreOptional);
 	AndroidThunkJava_IapQueryInAppPurchases = FindMethod(Env, GoogleServicesClassID, "AndroidThunkJava_IapQueryInAppPurchases", "([Ljava/lang/String;)Z", bIsStoreOptional);
 	AndroidThunkJava_IapBeginPurchase = FindMethod(Env, GoogleServicesClassID, "AndroidThunkJava_IapBeginPurchase", "([Ljava/lang/String;Ljava/lang/String;)Z", bIsStoreOptional);
 	AndroidThunkJava_IapIsAllowedToMakePurchases = FindMethod(Env, GoogleServicesClassID, "AndroidThunkJava_IapIsAllowedToMakePurchases", "()Z", bIsStoreOptional);
@@ -243,6 +243,8 @@ void FJavaWrapper::FindGooglePlayBillingMethods(JNIEnv* Env)
 	AndroidThunkJava_IapConsumePurchase = FindMethod(Env, GoogleServicesClassID, "AndroidThunkJava_IapConsumePurchase", "(Ljava/lang/String;)Z", bIsStoreOptional);
 	AndroidThunkJava_IapQueryExistingPurchases = FindMethod(Env, GoogleServicesClassID, "AndroidThunkJava_IapQueryExistingPurchases", "()Z", bIsStoreOptional);
 	AndroidThunkJava_IapGetBillingCountryCode = FindMethod(Env, GoogleServicesClassID, "AndroidThunkJava_IapGetBillingCountryCode", "()Ljava/lang/String;", bIsStoreOptional);
+	AndroidThunkJava_IapGetIsAlternativeBillingOnlyAvailable = FindMethod(Env, GoogleServicesClassID, "AndroidThunkJava_IapGetIsAlternativeBillingOnlyAvailable", "()Z", bIsStoreOptional);
+	AndroidThunkJava_IapShowAlternativeBillingOnlyInformationDialog = FindMethod(Env, GoogleServicesClassID, "AndroidThunkJava_IapShowAlternativeBillingOnlyInformationDialog", "()Z", bIsStoreOptional);
 }
 
 jclass FJavaWrapper::FindClass(JNIEnv* Env, const ANSICHAR* ClassName, bool bIsOptional)
@@ -590,6 +592,8 @@ jmethodID FJavaWrapper::AndroidThunkJava_IapQueryExistingPurchases;
 jmethodID FJavaWrapper::AndroidThunkJava_IapAcknowledgePurchase;
 jmethodID FJavaWrapper::AndroidThunkJava_IapConsumePurchase;
 jmethodID FJavaWrapper::AndroidThunkJava_IapGetBillingCountryCode;
+jmethodID FJavaWrapper::AndroidThunkJava_IapGetIsAlternativeBillingOnlyAvailable;
+jmethodID FJavaWrapper::AndroidThunkJava_IapShowAlternativeBillingOnlyInformationDialog;
 
 jmethodID FJavaWrapper::AndroidThunkJava_UseSurfaceViewWorkaround;
 jmethodID FJavaWrapper::AndroidThunkJava_SetDesiredViewSize;
@@ -1681,12 +1685,13 @@ bool AndroidThunkCpp_IsMusicActive()
 	return bIsActive;
 }
 
-void AndroidThunkCpp_Iap_SetupIapService(const FString& InProductKey)
+void AndroidThunkCpp_Iap_SetupIapService(const FString& InProductKey, bool InEnableAlternativeBillingOnly)
 {
 	if (JNIEnv* Env = FAndroidApplication::GetJavaEnv())
 	{
 		auto ProductKey = FJavaHelper::ToJavaString(Env, InProductKey);
-		FJavaWrapper::CallVoidMethod(Env, FJavaWrapper::GoogleServicesThis, FJavaWrapper::AndroidThunkJava_IapSetupService, *ProductKey);
+		jboolean EnableAlternativeBillingOnly = InEnableAlternativeBillingOnly ? JNI_TRUE : JNI_FALSE;
+		FJavaWrapper::CallVoidMethod(Env, FJavaWrapper::GoogleServicesThis, FJavaWrapper::AndroidThunkJava_IapSetupService, *ProductKey, EnableAlternativeBillingOnly);
 	}
 }
 
@@ -1829,6 +1834,32 @@ FString AndroidThunkJava_IapGetBillingCountryCode()
 	return Result;
 }
 
+bool AndroidThunkJava_IapGetIsAlternativeBillingOnlyAvailable()
+{
+	bool Result = false;
+	if (JNIEnv* Env = FAndroidApplication::GetJavaEnv())
+	{
+		CHECK_JNI_METHOD(FJavaWrapper::AndroidThunkJava_IapGetIsAlternativeBillingOnlyAvailable)
+
+		Result = FJavaWrapper::CallBooleanMethod(Env, FJavaWrapper::GameActivityThis, FJavaWrapper::AndroidThunkJava_IapGetIsAlternativeBillingOnlyAvailable);
+	}
+
+	return Result;
+}
+
+bool AndroidThunkJava_IapShowAlternativeBillingOnlyInformationDialog()
+{
+	bool Result = false;
+	if (JNIEnv* Env = FAndroidApplication::GetJavaEnv())
+	{
+		CHECK_JNI_METHOD(FJavaWrapper::AndroidThunkJava_IapShowAlternativeBillingOnlyInformationDialog)
+
+		Result = FJavaWrapper::CallBooleanMethod(Env, FJavaWrapper::GameActivityThis, FJavaWrapper::AndroidThunkJava_IapShowAlternativeBillingOnlyInformationDialog);
+	}
+
+	return Result;
+}
+
 void AndroidThunkCpp_UseSurfaceViewWorkaround()
 {
 	if (JNIEnv* Env = FAndroidApplication::GetJavaEnv())
diff --git a/Engine/Source/Runtime/Launch/Public/Android/AndroidJNI.h b/Engine/Source/Runtime/Launch/Public/Android/AndroidJNI.h
--- a/Engine/Source/Runtime/Launch/Public/Android/AndroidJNI.h
+++ b/Engine/Source/Runtime/Launch/Public/Android/AndroidJNI.h
@@ -130,6 +130,8 @@ public:
 	static jmethodID AndroidThunkJava_IapAcknowledgePurchase;
 	static jmethodID AndroidThunkJava_IapConsumePurchase;
 	static jmethodID AndroidThunkJava_IapGetBillingCountryCode;
+	static jmethodID AndroidThunkJava_IapGetIsAlternativeBillingOnlyAvailable;
+	static jmethodID AndroidThunkJava_IapShowAlternativeBillingOnlyInformationDialog;
 
 	// SurfaceView functionality for view scaling on some devices
 	static jmethodID AndroidThunkJava_UseSurfaceViewWorkaround;
