FIx LexToString for EOnlineKeyValuePairDataType.  This was probably silently converting to bool before CL 52404019

#Synced-CL 52313638

diff --git a/Engine/Plugins/Online/Android/OnlineSubsystemGooglePlay/Source/Private/OnlineLeaderboardInterfaceGooglePlay.cpp b/Engine/Plugins/Online/Android/OnlineSubsystemGooglePlay/Source/Private/OnlineLeaderboardInterfaceGooglePlay.cpp
--- a/Engine/Plugins/Online/Android/OnlineSubsystemGooglePlay/Source/Private/OnlineLeaderboardInterfaceGooglePlay.cpp
+++ b/Engine/Plugins/Online/Android/OnlineSubsystemGooglePlay/Source/Private/OnlineLeaderboardInterfaceGooglePlay.cpp
@@ -116,11 +116,11 @@ bool FOnlineLeaderboardsGooglePlay::WriteLeaderboards(const FName& SessionName,
 				}
 				default:
 				{
-					UE_LOG_ONLINE_LEADERBOARD(Warning, TEXT("Unsupported stat type %s"), *LexToString(Stat.GetType()));
+					UE_LOG_ONLINE_LEADERBOARD(Warning, TEXT("Unsupported stat type %s"), EOnlineKeyValuePairDataType::ToString(Stat.GetType()));
 					continue;
 				}
 			}
-			UE_LOG_ONLINE_LEADERBOARD(Display, TEXT("FOnlineLeaderboardsGooglePlay::WriteLeaderboards() %s value Score: %s"), *LexToString(Stat.GetType()), *Stat.ToString());
+			UE_LOG_ONLINE_LEADERBOARD(Display, TEXT("FOnlineLeaderboardsGooglePlay::WriteLeaderboards() %s value Score: %s"), EOnlineKeyValuePairDataType::ToString(Stat.GetType()), *Stat.ToString());
 			FGooglePlayLeaderboardScore& UnreportedScore = UnreportedScores.Emplace_GetRef();
 			UnreportedScore.GooglePlayLeaderboardId = MoveTemp(*GooglePlayLeaderboardId);
 			UnreportedScore.Score = Score;
